redesigned the session browser to be more space efficient and easier to use

diff --git a/Engine/Source/Developer/SessionFrontend/Private/SessionFrontendPrivatePCH.h b/Engine/Source/Developer/SessionFrontend/Private/SessionFrontendPrivatePCH.h
--- a/Engine/Source/Developer/SessionFrontend/Private/SessionFrontendPrivatePCH.h
+++ b/Engine/Source/Developer/SessionFrontend/Private/SessionFrontendPrivatePCH.h
@@ -35,7 +35,6 @@
 #include "SSessionBrowserInstanceListRow.h"
 #include "SSessionBrowserInstanceList.h"
 #include "SSessionBrowserSessionListRow.h"
-#include "SSessionBrowserContextMenu.h"
 #include "SSessionBrowserCommandBar.h"
 #include "SSessionBrowser.h"
 
diff --git a/Engine/Source/Developer/SessionFrontend/Private/Widgets/Browser/SSessionBrowser.cpp b/Engine/Source/Developer/SessionFrontend/Private/Widgets/Browser/SSessionBrowser.cpp
--- a/Engine/Source/Developer/SessionFrontend/Private/Widgets/Browser/SSessionBrowser.cpp
+++ b/Engine/Source/Developer/SessionFrontend/Private/Widgets/Browser/SSessionBrowser.cpp
@@ -35,64 +35,91 @@ void SSessionBrowser::Construct( const FArguments& InArgs, ISessionManagerRef In
 		SNew(SVerticalBox)
 
 		+ SVerticalBox::Slot()
-			.FillHeight(1.0)
+			.AutoHeight()
 			[
-				SAssignNew(Splitter, SSplitter)
+				SNew(SHorizontalBox)
 
-				+ SSplitter::Slot()
-					.Value(0.4)
+				+ SHorizontalBox::Slot()
+					.AutoWidth()
 					[
-						SNew(SBorder)
-							.BorderImage(FEditorStyle::GetBrush("ToolPanel.GroupBorder"))
-							.Padding(0.0f)
+						// session combo box
+						SAssignNew(SessionComboBox, SComboBox<ISessionInfoPtr>)
+							.ContentPadding(FMargin(6.0f, 2.0f))
+							.OptionsSource(&SessionList)
+							.ToolTipText(LOCTEXT("SessionComboBoxToolTip", "Select the game session to interact with.").ToString())
+							.OnGenerateWidget(this, &SSessionBrowser::HandleSessionComboBoxGenerateWidget)
+							.OnSelectionChanged(this, &SSessionBrowser::HandleSessionComboBoxSelectionChanged)
 							[
-								// session list
-								SAssignNew(SessionListView, SListView<ISessionInfoPtr>)
-									.ItemHeight(24.0f)
-									.ListItemsSource(&SessionList)
-									.OnContextMenuOpening(this, &SSessionBrowser::HandleSessionListViewContextMenuOpening)
-									.OnGenerateRow(this, &SSessionBrowser::HandleSessionListViewGenerateRow)
-									.OnSelectionChanged(this, &SSessionBrowser::HandleSessionListViewSelectionChanged)
-									.HeaderRow
-									(
-										SNew(SHeaderRow)
-
-										+ SHeaderRow::Column("Status")
-											.DefaultLabel(FText::FromString(" "))
-											.FixedWidth(24.0f)
-
-										+ SHeaderRow::Column("Name")
-											.DefaultLabel(LOCTEXT("SessionListNameColumnTitle", "Session Name"))
-											.FillWidth(0.4f)
-
-										+ SHeaderRow::Column("Owner")
-											.DefaultLabel(LOCTEXT("SessionListOwnerColumnTitle", "Owner"))
-											.FillWidth(0.6f)
-									)
+								SNew(STextBlock)
+									.Font(FSlateFontInfo(FPaths::EngineContentDir() / TEXT("Slate/Fonts/Roboto-Bold.ttf"), 13))
+									.Text(this, &SSessionBrowser::HandleSessionComboBoxText)
 							]
 					]
 
-				+ SSplitter::Slot()
-					.Value(0.5)
+				+ SHorizontalBox::Slot()
+					.FillWidth(1.0f)
+					.Padding(4.0f, 0.0f, 0.0f, 0.0f)
+					.VAlign(VAlign_Center)
+					[
+						SNew(STextBlock)
+							.Text(this, &SSessionBrowser::HandleSessionDetailsText)
+					]
+
+				+ SHorizontalBox::Slot()
+					.AutoWidth()
+					.Padding(8.0f, 0.0f)
+					[
+						SNew(SSeparator)
+							.Orientation(Orient_Vertical)
+					]
+
+				+ SHorizontalBox::Slot()
+					.AutoWidth()
 					[
-						SNew(SBorder)
-							.BorderImage(FEditorStyle::GetBrush("ToolPanel.GroupBorder"))
-							.Padding(0.0f)
+						// terminate button
+						SNew(SButton)
+							.ButtonStyle(FEditorStyle::Get(), "ToggleButton")
+							.ContentPadding(FMargin(6.0f, 2.0f))
+							.IsEnabled(this, &SSessionBrowser::HandleTerminateSessionButtonIsEnabled)
+							.OnClicked(this, &SSessionBrowser::HandleTerminateSessionButtonClicked)
+							.ToolTipText(LOCTEXT("TerminateButtonTooltip", "Shuts down all game instances that are part of this session.").ToString())
 							[
-								// instance list
-								SAssignNew(InstanceListView, SSessionInstanceList, SessionManager.ToSharedRef())
-									.SelectionMode(ESelectionMode::Multi)
+								SNew(SHorizontalBox)
+
+								+ SHorizontalBox::Slot()
+									.AutoWidth()
+									.VAlign(VAlign_Center)
+									[
+										SNew(SImage)
+											.Image(FEditorStyle::GetBrush("SessionBrowser.Terminate"))
+									]
+
+								+ SHorizontalBox::Slot()
+									.AutoWidth()
+									.VAlign(VAlign_Center)
+									.Padding(4.0f, 1.0f, 0.0f, 0.0f)
+									[
+										SNew(STextBlock)
+											.TextStyle(FEditorStyle::Get(), "ContentBrowser.TopBar.Font")
+											.Text(LOCTEXT("TerminateSessionButtonLabel", "Terminate Session").ToString())
+									]
 							]
 					]
-
 			]
-		/*
+
 		+ SVerticalBox::Slot()
-			.AutoHeight()
+			.FillHeight(1.0f)
+			.Padding(0.0f, 4.0f, 0.0f, 0.0f)
 			[
-				// command bar
-				SAssignNew(CommandBar, SSessionBrowserCommandBar)
-			]*/
+				SNew(SBorder)
+					.BorderImage(FEditorStyle::GetBrush("ToolPanel.GroupBorder"))
+					.Padding(0.0f)
+					[
+						// instance list
+						SAssignNew(InstanceListView, SSessionInstanceList, SessionManager.ToSharedRef())
+							.SelectionMode(ESelectionMode::Multi)
+					]
+			]
 	];
 
 	SessionManager->OnSelectedSessionChanged().AddSP(this, &SSessionBrowser::HandleSessionManagerSelectedSessionChanged);
@@ -102,24 +129,6 @@ void SSessionBrowser::Construct( const FArguments& InArgs, ISessionManagerRef In
 }
 
 
-/* SWidget overrides
- *****************************************************************************/
-
-void SSessionBrowser::Tick( const FGeometry& AllottedGeometry, const double InCurrentTime, const float InDeltaTime )
-{
-	SCompoundWidget::Tick(AllottedGeometry, InCurrentTime, InDeltaTime);
-
-	if ((AllottedGeometry.Size.X > 512.0f) && (AllottedGeometry.Size.X > AllottedGeometry.Size.Y))
-	{
-		Splitter->SetOrientation(Orient_Horizontal);
-	}
-	else
-	{
-		Splitter->SetOrientation(Orient_Vertical);
-	}
-}
-
-
 /* SSessionBrowser implementation
  *****************************************************************************/
 
@@ -137,8 +146,48 @@ void SSessionBrowser::FilterSessions( )
 		}
 	}
 
-	SessionListView->RequestListRefresh();
-	SessionListView->SetSelection(SessionManager->GetSelectedSession());
+	SessionComboBox->RefreshOptions();
+	SessionComboBox->SetSelectedItem(SessionManager->GetSelectedSession());
+}
+
+
+FText SSessionBrowser::GetSessionName( const ISessionInfoPtr& SessionInfo ) const
+{
+	FText SessionName;
+
+	if (!SessionInfo->IsStandalone() || !SessionInfo->GetSessionName().IsEmpty())
+	{
+		SessionName = FText::FromString(SessionInfo->GetSessionName());
+	}
+	else
+	{
+		TArray<ISessionInstanceInfoPtr> Instances;
+		SessionInfo->GetInstances(Instances);
+
+		if (Instances.Num() > 0)
+		{
+			const ISessionInstanceInfoPtr& FirstInstance = Instances[0];
+
+			if ((Instances.Num() == 1) && (FirstInstance->GetInstanceId() == FApp::GetInstanceId()))
+			{
+				SessionName = LOCTEXT("ThisApplicationSessionText", "This Application");
+			}
+			else if (FirstInstance->GetDeviceName() == FPlatformProcess::ComputerName())
+			{
+				SessionName = LOCTEXT("UnnamedLocalSessionText", "Unnamed Session (Local)");
+			}
+			else
+			{
+				SessionName = LOCTEXT("UnnamedRemoteSessionText", "Unnamed Session (Remote)");
+			}
+		}
+		else
+		{
+			SessionName = LOCTEXT("UnnamedSessionText", "Unnamed Session");
+		}
+	}
+
+	return SessionName;
 }
 
 
@@ -157,16 +206,81 @@ void SSessionBrowser::ReloadSessions( )
 /* SSessionBrowser event handlers
  *****************************************************************************/
 
-TSharedPtr<SWidget> SSessionBrowser::HandleSessionListViewContextMenuOpening( )
+TSharedRef<SWidget> SSessionBrowser::HandleSessionComboBoxGenerateWidget( ISessionInfoPtr SessionInfo ) const
+{
+	return SNew(SVerticalBox)
+
+	+ SVerticalBox::Slot()
+		.AutoHeight()
+		.Padding(4.0f, 0.0f)
+		[
+			SNew(STextBlock)
+				.Font(FSlateFontInfo(FPaths::EngineContentDir() / TEXT("Slate/Fonts/Roboto-Bold.ttf"), 14))
+				.Text(GetSessionName(SessionInfo))
+		]
+
+	+ SVerticalBox::Slot()
+		.AutoHeight()
+		.Padding(4.0f, 0.0f)
+		[
+			SNew(SHorizontalBox)
+
+			+ SHorizontalBox::Slot()
+				.AutoWidth()
+				.HAlign(HAlign_Left)
+				[
+					SNew(STextBlock)
+						.Text(SessionInfo->GetSessionOwner())
+				]
+
+			+ SHorizontalBox::Slot()
+				.AutoWidth()
+				.HAlign(HAlign_Right)
+				.Padding(16.0f, 0.0f, 0.0f, 0.0f)
+				[
+					SNew(STextBlock)
+						.Text(FText::Format(LOCTEXT("InstanceCountFormat", "{0} Instance(s)"), FText::AsNumber(SessionInfo->GetNumInstances())))
+				]
+		];
+}
+
+
+void SSessionBrowser::HandleSessionComboBoxSelectionChanged( ISessionInfoPtr SelectedItem, ESelectInfo::Type SelectInfo )
+{
+	if (!SessionManager->SelectSession(SelectedItem))
+	{
+		SessionComboBox->SetSelectedItem(SessionManager->GetSelectedSession());
+	}
+}
+
+
+FText SSessionBrowser::HandleSessionComboBoxText( ) const
 {
 	const ISessionInfoPtr& SelectedSession = SessionManager->GetSelectedSession();
 
-	if (SelectedSession.IsValid() && (SelectedSession->GetSessionOwner() == FPlatformProcess::UserName(false)))
+	if (SelectedSession.IsValid())
 	{
-		return SNew(SSessionBrowserContextMenu, SessionManager.ToSharedRef());
-	}		
-	
-	return NULL;
+		return GetSessionName(SelectedSession);
+	}
+
+	return LOCTEXT("SelectSessionHint", "Select a session...");
+}
+
+
+FText SSessionBrowser::HandleSessionDetailsText( ) const
+{
+	const ISessionInfoPtr& SelectedSession = SessionManager->GetSelectedSession();
+
+	if (!SelectedSession.IsValid())
+	{
+		return FText::GetEmpty();
+	}
+
+	FFormatNamedArguments FormatArguments;
+	FormatArguments.Add(TEXT("NumInstances"), SelectedSession->GetNumInstances());
+	FormatArguments.Add(TEXT("OwnerName"), SelectedSession->GetSessionOwner().IsEmpty() ? LOCTEXT("UnknownOwner", "<unknown>") : FText::FromString(SelectedSession->GetSessionOwner()));
+
+	return FText::Format(LOCTEXT("SessionDetails", "Owner: {OwnerName}, {NumInstances} Instance(s)"), FormatArguments);
 }
 
 
@@ -189,9 +303,9 @@ void SSessionBrowser::HandleSessionListViewSelectionChanged( ISessionInfoPtr Ite
 
 void SSessionBrowser::HandleSessionManagerSelectedSessionChanged( const ISessionInfoPtr& SelectedSession )
 {
-	if (!SessionListView->IsItemSelected(SelectedSession))
+	if (SessionComboBox->GetSelectedItem() != SelectedSession)
 	{
-		SessionListView->SetSelection(SelectedSession);
+		SessionComboBox->SetSelectedItem(SelectedSession);
 	}
 }
 
@@ -202,4 +316,35 @@ void SSessionBrowser::HandleSessionManagerSessionsUpdated( )
 }
 
 
+FReply SSessionBrowser::HandleTerminateSessionButtonClicked( )
+{
+	int32 DialogResult = FMessageDialog::Open(EAppMsgType::YesNo, LOCTEXT("TerminateSessionDialogPrompt", "Are you sure you want to terminate this session and its instances?"));
+
+	if (DialogResult == EAppReturnType::Yes)
+	{
+		const ISessionInfoPtr& SelectedSession = SessionManager->GetSelectedSession();
+
+		if (SelectedSession.IsValid())
+		{
+			if (SelectedSession->GetSessionOwner() == FPlatformProcess::UserName(false))
+			{
+				SelectedSession->Terminate();
+			}
+			else
+			{
+				FMessageDialog::Open( EAppMsgType::Ok, FText::Format( LOCTEXT("TerminateDeniedPrompt", "You are not authorized to terminate the currently selected session, because it is owned by {0}"), FText::FromString( SelectedSession->GetSessionOwner() ) ) );
+			}
+		}
+	}
+
+	return FReply::Handled();
+}
+
+
+bool SSessionBrowser::HandleTerminateSessionButtonIsEnabled( ) const
+{
+	return SessionManager->GetSelectedSession().IsValid();
+}
+
+
 #undef LOCTEXT_NAMESPACE
diff --git a/Engine/Source/Developer/SessionFrontend/Private/Widgets/Browser/SSessionBrowser.h b/Engine/Source/Developer/SessionFrontend/Private/Widgets/Browser/SSessionBrowser.h
--- a/Engine/Source/Developer/SessionFrontend/Private/Widgets/Browser/SSessionBrowser.h
+++ b/Engine/Source/Developer/SessionFrontend/Private/Widgets/Browser/SSessionBrowser.h
@@ -35,14 +35,6 @@ public:
 	 */
 	void Construct( const FArguments& InArgs, ISessionManagerRef InSessionManager );
 
-public:
-
-	// Begin SWidget interface
-
-	virtual void Tick( const FGeometry& AllottedGeometry, const double InCurrentTime, const float InDeltaTime ) OVERRIDE;
-
-	// End SWidget interface
-
 protected:
 
 	/**
@@ -50,6 +42,15 @@ protected:
 	 */
 	void FilterSessions( );
 
+	/**
+	 * Gets the display name of the specified session.
+	 *
+	 * @param SessionInfo - The session to get the name for.
+	 *
+	 * @param Session name.
+	 */
+	FText GetSessionName( const ISessionInfoPtr& SessionInfo ) const;
+
 	/**
 	 * Reloads the sessions list.
 	 */
@@ -57,21 +58,30 @@ protected:
 
 private:
 
-	// Callback for generating a context menu in the session list view.
-	TSharedPtr<SWidget> HandleSessionListViewContextMenuOpening( );
+	// Callback for generating a row in the session combo box.
+	TSharedRef<SWidget> HandleSessionComboBoxGenerateWidget( ISessionInfoPtr SessionInfo ) const;
 
-	// Callback for creating a row widget in the session list view.
-	TSharedRef<ITableRow> HandleSessionListViewGenerateRow( ISessionInfoPtr Item, const TSharedRef<STableViewBase>& OwnerTable );
+	// Callback for selection changes in the session combo box.
+	void HandleSessionComboBoxSelectionChanged( ISessionInfoPtr SelectedItem, ESelectInfo::Type SelectInfo );
 
-	// Handles changing the selection in the session tree.
-	void HandleSessionListViewSelectionChanged( ISessionInfoPtr Item, ESelectInfo::Type SelectInfo );
+	// Callback for getting the text in the session combo box.
+	FText HandleSessionComboBoxText( ) const;
 
-	// Handles changing the selected session in the session manager.
+	// Callback for getting the text of the session details area.
+	FText HandleSessionDetailsText( ) const;
+
+	// Callback for changing the selected session in the session manager.
 	void HandleSessionManagerSelectedSessionChanged( const ISessionInfoPtr& SelectedSession );
 
-	// Handles updating the session list in the session manager.
+	// Callback for updating the session list in the session manager.
 	void HandleSessionManagerSessionsUpdated( );
 
+	// Callback for clicking the 'Terminate Session' button.
+	FReply HandleTerminateSessionButtonClicked( );
+
+	// Callback for getting the enabled state of the 'Terminate Session' button.
+	bool HandleTerminateSessionButtonIsEnabled( ) const;
+
 private:
 
 	// Holds an unfiltered list of available sessions.
@@ -83,15 +93,23 @@ private:
 	// Holds the instance list view.
 	TSharedPtr<SSessionInstanceList> InstanceListView;
 
+	// Holds the session combo box.
+	TSharedPtr<SComboBox<ISessionInfoPtr> > SessionComboBox;
+
 	// Holds the filtered list of sessions to be displayed.
 	TArray<ISessionInfoPtr> SessionList;
 
-	// Holds the session list view.
-	TSharedPtr<SListView<ISessionInfoPtr> > SessionListView;
-
 	// Holds a reference to the session manager.
 	ISessionManagerPtr SessionManager;
 
-	// Holds the splitter.
-	TSharedPtr<SSplitter> Splitter;
+private:
+
+	// Callback for creating a row widget in the session list view.
+	TSharedRef<ITableRow> HandleSessionListViewGenerateRow( ISessionInfoPtr Item, const TSharedRef<STableViewBase>& OwnerTable );
+
+	// Handles changing the selection in the session tree.
+	void HandleSessionListViewSelectionChanged( ISessionInfoPtr Item, ESelectInfo::Type SelectInfo );
+
+	// Holds the session list view.
+	TSharedPtr<SListView<ISessionInfoPtr> > SessionListView;
 };
diff --git a/Engine/Source/Developer/SessionFrontend/Private/Widgets/Browser/SSessionBrowserContextMenu.h b/Engine/Source/Developer/SessionFrontend/Private/Widgets/Browser/SSessionBrowserContextMenu.h
deleted file mode 100644
diff --git a/Engine/Source/Developer/SessionFrontend/Private/Widgets/Browser/SSessionBrowserInstanceList.h b/Engine/Source/Developer/SessionFrontend/Private/Widgets/Browser/SSessionBrowserInstanceList.h
--- a/Engine/Source/Developer/SessionFrontend/Private/Widgets/Browser/SSessionBrowserInstanceList.h
+++ b/Engine/Source/Developer/SessionFrontend/Private/Widgets/Browser/SSessionBrowserInstanceList.h
@@ -22,7 +22,6 @@ public:
 		SLATE_ATTRIBUTE(ESelectionMode::Type, SelectionMode)
 	SLATE_END_ARGS()
 
-
 public:
 
 	/**
@@ -37,7 +36,6 @@ public:
 		}
 	}
 
-
 public:
 	
 	/**
@@ -63,16 +61,22 @@ public:
 				(
 					SNew(SHeaderRow)
 
-					+ SHeaderRow::Column("Icons")
-						.DefaultLabel(LOCTEXT("InstanceListStatusColumnHeader", " ").ToString())
-						.FixedWidth(56.0f)
-
-					+ SHeaderRow::Column("Instance")
+					+ SHeaderRow::Column("Name")
 						.DefaultLabel(LOCTEXT("InstanceListNameColumnHeader", "Instance").ToString())
-/*
-					+ SHeaderRow::Column("Level")
-						.DefaultLabel(LOCTEXT("InstanceListLevelColumnHeader", "Current Level").ToString())*/
 
+					+ SHeaderRow::Column("Type")
+						.DefaultLabel(LOCTEXT("InstanceListTypeColumnHeader", "Type").ToString())
+
+					+ SHeaderRow::Column("Device")
+						.DefaultLabel(LOCTEXT("InstanceListDeviceColumnHeader", "Device").ToString())
+
+					+ SHeaderRow::Column("Platform")
+						.DefaultLabel(LOCTEXT("InstanceListPlatformColumnHeader", "Platform").ToString())
+
+					+ SHeaderRow::Column("Status")
+						.DefaultLabel(LOCTEXT("InstanceListStatusColumnHeader", "Status").ToString())
+						.HAlignCell(HAlign_Right)
+						.HAlignHeader(HAlign_Right)
 				)
 		];
 
@@ -80,7 +84,6 @@ public:
 		SessionManager->OnSelectedSessionChanged().AddSP(this, &SSessionInstanceList::HandleSessionManagerSelectedSessionChanged);
 	}
 
-
 protected:
 
 	/**
@@ -111,24 +114,19 @@ protected:
 		}
 	}
 
-
 private:
 
 	// Handles getting the tool tip text for an instance.
-	FString HandleInstanceGetToolTipText( ISessionInstanceInfoPtr Instance ) const
+	FText HandleInstanceGetToolTipText( ISessionInstanceInfoPtr Instance ) const
 	{
-		FString ToolTipText;
-
-		ToolTipText += LOCTEXT("InstanceToolTipInstanceType", "Instance Type: ").ToString() + Instance->GetInstanceType() + LINE_TERMINATOR;
-		ToolTipText += LOCTEXT("InstanceToolTipPlatform", "Platform: ").ToString() + Instance->GetPlatformName() + LINE_TERMINATOR;
-		ToolTipText += LOCTEXT("InstanceToolTipBuildDate", "Build Date: ").ToString() + Instance->GetBuildDate() + LINE_TERMINATOR;
-		ToolTipText += LOCTEXT("InstanceToolTipEngineVersion", "Engine Version: ").ToString() + FString::FromInt(Instance->GetEngineVersion()) + LINE_TERMINATOR;
-		ToolTipText += LOCTEXT("InstanceToolTipConsoleBuild", "Console Build: ").ToString() + (Instance->IsConsole() ? TEXT("true") : TEXT("false")) + LINE_TERMINATOR;
-		ToolTipText += LOCTEXT("InstanceToolTipDeviceName", "Device Name: ").ToString() + Instance->GetDeviceName() + LINE_TERMINATOR;
-		ToolTipText += LOCTEXT("InstanceToolTipUserName", "Owner: ").ToString() + Instance->GetOwnerSession()->GetSessionOwner() + LINE_TERMINATOR;
-		ToolTipText += LOCTEXT("InstanceToolTipInstanceId", "ID:").ToString() + Instance->GetInstanceId().ToString(EGuidFormats::DigitsWithHyphensInBraces);
-
-		return ToolTipText;
+		FTextBuilder ToolTipTextBuilder;
+
+		ToolTipTextBuilder.AppendLineFormat(LOCTEXT("InstanceToolTipBuildDate", "Build Date: {0}"), FText::FromString(Instance->GetBuildDate()));
+		ToolTipTextBuilder.AppendLineFormat(LOCTEXT("InstanceToolTipConsoleBuild", "Console Build: {0}"), Instance->IsConsole() ? LOCTEXT("LabelYes", "Yes") : LOCTEXT("LabelNo", "No"));
+		ToolTipTextBuilder.AppendLineFormat(LOCTEXT("InstanceToolTipEngineVersion", "Engine Version: {0}"), FText::FromString(FString::FromInt(Instance->GetEngineVersion())));
+		ToolTipTextBuilder.AppendLineFormat(LOCTEXT("InstanceToolTipInstanceId", "Instance ID: {0}"), FText::FromString(Instance->GetInstanceId().ToString(EGuidFormats::DigitsWithHyphensInBraces)));
+
+		return ToolTipTextBuilder.ToText();
 	}
 
 	// Handles creating a row widget for the instance list view.
@@ -169,7 +167,6 @@ private:
 		UpdateList();
 	}
 
-
 private:
 
 	// Holds a flag indicating whether instance selection changes should be ignored.
diff --git a/Engine/Source/Developer/SessionFrontend/Private/Widgets/Browser/SSessionBrowserInstanceListRow.h b/Engine/Source/Developer/SessionFrontend/Private/Widgets/Browser/SSessionBrowserInstanceListRow.h
--- a/Engine/Source/Developer/SessionFrontend/Private/Widgets/Browser/SSessionBrowserInstanceListRow.h
+++ b/Engine/Source/Developer/SessionFrontend/Private/Widgets/Browser/SSessionBrowserInstanceListRow.h
@@ -7,6 +7,9 @@
 #pragma once
 
 
+#define LOCTEXT_NAMESPACE "SSessionBrowserInstanceListRow"
+
+
 /**
  * Delegate type for changed instance check box state changes.
  *
@@ -28,7 +31,6 @@ public:
 		SLATE_ARGUMENT(ISessionInstanceInfoPtr, InstanceInfo)
 	SLATE_END_ARGS()
 
-
 public:
 
 	/**
@@ -43,7 +45,6 @@ public:
 		SMultiColumnTableRow<ISessionInstanceInfoPtr>::Construct(FSuperRowType::FArguments(), InOwnerTableView);
 	}
 
-
 public:
 
 	/**
@@ -55,37 +56,39 @@ public:
 	 */
 	virtual TSharedRef<SWidget> GenerateWidgetForColumn( const FName& ColumnName ) OVERRIDE
 	{
-		if (ColumnName == "Icons")
+		if (ColumnName == "Device")
 		{
-			return SNew(SHorizontalBox)
-
-			+ SHorizontalBox::Slot()
-				.AutoWidth()
-				.Padding(1.0)
+			return SNew(SBox)
+				.Padding(FMargin(4.0f, 0.0f))
+				.VAlign(VAlign_Center)
 				[
-					SNew(SImage)
-						.Image(FEditorStyle::GetBrush(*FString::Printf(TEXT("Launcher.Platform_%s"), *InstanceInfo->GetPlatformName())))
-				]
-
-			+ SHorizontalBox::Slot()
-				.AutoWidth()
-				.Padding(1.0)
+					SNew(STextBlock)
+						.ColorAndOpacity(this, &SSessionBrowserInstanceListRow::HandleTextColorAndOpacity)
+						.Text(this, &SSessionBrowserInstanceListRow::HandleDeviceColumnText)
+				];
+		}
+		else if (ColumnName == "Level")
+		{
+			return SNew(SBox)
+				.Padding(FMargin(4.0, 0.0))
+				.VAlign(VAlign_Center)
 				[
-					SNew(SImage)
-						.Image(this, &SSessionBrowserInstanceListRow::HandleGetInstanceIcon)
+					SNew(STextBlock)
+						.ColorAndOpacity(this, &SSessionBrowserInstanceListRow::HandleTextColorAndOpacity)
+						.Text(this, &SSessionBrowserInstanceListRow::HandleLevelColumnText)
 				];
 		}
-		else if (ColumnName == "Instance")
+		else if (ColumnName == "Name")
 		{
 			return SNew(SBox)
-				.Padding(FMargin(4.0f, 1.0f, 4.0f, 0.0f))
+				.Padding(FMargin(1.0f, 1.0f, 4.0f, 0.0f))
 				.HAlign(HAlign_Left)
 				[
 					SNew(SBorder)
-						.BorderBackgroundColor(this, &SSessionBrowserInstanceListRow::HandleBorderGetBackgroundColor)
-						.BorderImage(this, &SSessionBrowserInstanceListRow::HandleBorderGetBrush)
+						.BorderBackgroundColor(this, &SSessionBrowserInstanceListRow::HandleInstanceBorderBackgroundColor)
+						.BorderImage(this, &SSessionBrowserInstanceListRow::HandleInstanceBorderBrush)
 						.ColorAndOpacity(FLinearColor(0.25f, 0.25f, 0.25f))
-						.Padding(FMargin(6.0, 4.0))
+						.Padding(FMargin(6.0f, 4.0f))
 						.VAlign(VAlign_Center)
 						[
 							SNew(STextBlock)
@@ -94,32 +97,68 @@ public:
 						]
 				];
 		}
-		else if (ColumnName == "Level")
+		else if (ColumnName == "Platform")
+		{
+			return SNew(SHorizontalBox)
+
+			+ SHorizontalBox::Slot()
+				.AutoWidth()
+				.VAlign(VAlign_Center)
+				[
+					SNew(SImage)
+						.Image(FEditorStyle::GetBrush(*FString::Printf(TEXT("Launcher.Platform_%s"), *InstanceInfo->GetPlatformName())))
+				]
+
+			+ SHorizontalBox::Slot()
+				.FillWidth(1.0f)
+				.Padding(4.0f, 0.0f, 0.0f, 0.0f)
+				.VAlign(VAlign_Center)
+				[
+					SNew(STextBlock)
+						.Text(FText::FromString(InstanceInfo->GetPlatformName()))
+				];
+		}
+		else if (ColumnName == "Status")
 		{
 			return SNew(SBox)
-				.Padding(FMargin(4.0, 0.0))
+				.Padding(FMargin(4.0f, 0.0f))
 				.VAlign(VAlign_Center)
 				[
 					SNew(STextBlock)
-						.ColorAndOpacity(this, &SSessionBrowserInstanceListRow::HandleGetColorAndOpacity)
-						.Text(this, &SSessionBrowserInstanceListRow::HandleGetLevelText)
+						.ColorAndOpacity(this, &SSessionBrowserInstanceListRow::HandleTextColorAndOpacity)
+						.Text(this, &SSessionBrowserInstanceListRow::HandleStatusColumnText)
+				];
+		}
+		else if (ColumnName == "Type")
+		{
+			return SNew(SBox)
+				.Padding(FMargin(4.0f, 0.0f))
+				.VAlign(VAlign_Center)
+				[
+					SNew(STextBlock)
+						.Text(FText::FromString(InstanceInfo->GetInstanceType()))
 				];
 		}
 
 		return SNullWidget::NullWidget;
 	}
 
-
 private:
 
+	// Callback for getting the text in the 'Device' column.
+	FText HandleDeviceColumnText( ) const
+	{
+		return FText::FromString(InstanceInfo->GetDeviceName());
+	}
+
 	// Callback for getting the border color for this row.
-	FSlateColor HandleBorderGetBackgroundColor( ) const
+	FSlateColor HandleInstanceBorderBackgroundColor( ) const
 	{
 		return FLinearColor((GetTypeHash(InstanceInfo->GetInstanceId()) & 0xff) * 360.0 / 256.0, 0.8, 0.3, 1.0).HSVToLinearRGB();
 	}
 
 	// Callback for getting the border brush for this row.
-	const FSlateBrush* HandleBorderGetBrush( ) const
+	const FSlateBrush* HandleInstanceBorderBrush( ) const
 	{
 		if (FDateTime::UtcNow() - InstanceInfo->GetLastUpdateTime() < FTimespan::FromSeconds(10.0))
 		{
@@ -129,30 +168,34 @@ private:
 		return FEditorStyle::GetBrush("ErrorReporting.EmptyBox");
 	}
 
-	// Callback for getting the foreground text color.
-	FSlateColor HandleGetColorAndOpacity( ) const
+	// Callback for getting the instance's current level.
+	FString HandleLevelColumnText( ) const
+	{
+		return InstanceInfo->GetCurrentLevel();
+	}
+
+	// Callback for getting the text in the 'Status' column.
+	FText HandleStatusColumnText( ) const
 	{
 		if (FDateTime::UtcNow() - InstanceInfo->GetLastUpdateTime() < FTimespan::FromSeconds(10.0))
 		{
-			return FSlateColor::UseForeground();
+			return LOCTEXT("StatusRunning", "Running");
 		}
 
-		return FSlateColor::UseSubduedForeground();
+		return LOCTEXT("StatusTimedOut", "Timed Out");
 	}
 
-	// Callback for getting the instance type icon.
-	const FSlateBrush* HandleGetInstanceIcon( ) const
+	// Callback for getting the foreground text color.
+	FSlateColor HandleTextColorAndOpacity( ) const
 	{
-		return FEditorStyle::GetBrush(*FString::Printf(TEXT("Launcher.Instance_%s"), *InstanceInfo->GetInstanceType()));
-	}
+		if (FDateTime::UtcNow() - InstanceInfo->GetLastUpdateTime() < FTimespan::FromSeconds(10.0))
+		{
+			return FSlateColor::UseForeground();
+		}
 
-	// Callback for getting the instance's current level.
-	FString HandleGetLevelText( ) const
-	{
-		return InstanceInfo->GetCurrentLevel();
+		return FSlateColor::UseSubduedForeground();
 	}
 
-
 private:
 
 	// Holds a reference to the instance info that is displayed in this row.
@@ -161,3 +204,6 @@ private:
 	// Holds a delegate to be invoked when the check box state changed.
 	FOnCheckStateChanged OnCheckStateChanged;
 };
+
+
+#undef LOCTEXT_NAMESPACE
diff --git a/Engine/Source/Developer/SessionFrontend/Private/Widgets/Browser/SSessionBrowserSessionListRow.h b/Engine/Source/Developer/SessionFrontend/Private/Widgets/Browser/SSessionBrowserSessionListRow.h
--- a/Engine/Source/Developer/SessionFrontend/Private/Widgets/Browser/SSessionBrowserSessionListRow.h
+++ b/Engine/Source/Developer/SessionFrontend/Private/Widgets/Browser/SSessionBrowserSessionListRow.h
@@ -24,7 +24,6 @@ public:
 		SLATE_ARGUMENT(ISessionInfoPtr, SessionInfo)
 	SLATE_END_ARGS()
 
-
 public:
 
 	/**
@@ -37,10 +36,9 @@ public:
 		HighlightText = InArgs._HighlightText;
 		SessionInfo = InArgs._SessionInfo;
 
-		SMultiColumnTableRow<ISessionInfoPtr>::Construct(FSuperRowType::FArguments(), InOwnerTableView);
+		SMultiColumnTableRow<ISessionInfoPtr>::Construct(FSuperRowType::FArguments().Style(FEditorStyle::Get(), "TableView.Row"), InOwnerTableView);
 	}
 
-
 public:
 
 	/**
@@ -54,30 +52,48 @@ public:
 	{
 		if (ColumnName == "Name")
 		{
-			return SNew(SBox)
-				.Padding(FMargin(4.0f, 0.0f))
-				.VAlign(VAlign_Center)
+			return SNew(SVerticalBox)
+
+			+ SVerticalBox::Slot()
+				.AutoHeight()
+				.Padding(8.0f, 0.0f, 0.0f, 0.0f)
 				[
 					SNew(STextBlock)
+						.Font(FSlateFontInfo(FPaths::EngineContentDir() / TEXT("Slate/Fonts/Roboto-Bold.ttf"), 14))
 						.HighlightText(HighlightText)
 						.Text(this, &SSessionBrowserSessionListRow::HandleGetSessionName)
+				]
+
+			+ SVerticalBox::Slot()
+				.AutoHeight()
+				.Padding(8.0f, 0.0f)
+				[
+					SNew(SHorizontalBox)
+
+					+ SHorizontalBox::Slot()
+						.HAlign(HAlign_Left)
+						[
+							SNew(STextBlock)
+								.Text(SessionInfo->GetSessionOwner())
+						]
+
+					+ SHorizontalBox::Slot()
+						.HAlign(HAlign_Right)
+						[
+							SNew(STextBlock)
+								.Text(LOCTEXT("adsf", "1 Instance(s)"))
+						]
 				];
 		}
-		else if (ColumnName == "Owner")
+		else if (ColumnName == "Selection")
 		{
-			return SNew(SBox)
-				.Padding(FMargin(4.0f, 0.0f))
-				.VAlign(VAlign_Center)
-				[
-					SNew(STextBlock)
-						.Text(SessionInfo->GetSessionOwner())
-				];
+			return SNew(SColorBlock)
+				.Color(FLinearColor::Gray);
 		}
 		
 		return SNullWidget::NullWidget;
 	}
 
-
 protected:
 
 	// Callback for getting the name of the session.
@@ -117,7 +133,6 @@ protected:
 		return LOCTEXT("UnnamedSessionText", "Unnamed Session").ToString();
 	}
 
-
 private:
 
 	// Holds the highlight string for the session name.
diff --git a/Engine/Source/Developer/SessionFrontend/Private/Widgets/SSessionFrontend.cpp b/Engine/Source/Developer/SessionFrontend/Private/Widgets/SSessionFrontend.cpp
--- a/Engine/Source/Developer/SessionFrontend/Private/Widgets/SSessionFrontend.cpp
+++ b/Engine/Source/Developer/SessionFrontend/Private/Widgets/SSessionFrontend.cpp
@@ -63,14 +63,14 @@ void SSessionFrontend::Construct( const FArguments& InArgs, const TSharedRef<SDo
 		->AddArea
 		(
 			FTabManager::NewPrimaryArea()
-				->SetOrientation(Orient_Horizontal)
+				->SetOrientation(Orient_Vertical)
 				->Split
 				(
 					// session browser
 					FTabManager::NewStack()
 						->AddTab(SessionBrowserTabId, ETabState::OpenedTab)
 						->SetHideTabWell(true)
-						->SetSizeCoefficient(0.25f)
+						->SetSizeCoefficient(0.15f)
 				)
 				->Split
 				(
@@ -80,7 +80,7 @@ void SSessionFrontend::Construct( const FArguments& InArgs, const TSharedRef<SDo
 						->AddTab(AutomationTabId, ETabState::OpenedTab)
 						->AddTab(SessionScreenTabId, ETabState::ClosedTab)
 						->AddTab(ProfilerTabId, ETabState::OpenedTab)
-						->SetSizeCoefficient(0.75f)
+						->SetSizeCoefficient(0.85f)
 						->SetForegroundTab(SessionConsoleTabId)
 				)							
 		);
