SessionFrontend: Documentation and formatting cleanup pass

diff --git a/Engine/Source/Developer/SessionFrontend/Private/Models/SessionBrowserOwnerFilter.h b/Engine/Source/Developer/SessionFrontend/Private/Models/SessionBrowserOwnerFilter.h
--- a/Engine/Source/Developer/SessionFrontend/Private/Models/SessionBrowserOwnerFilter.h
+++ b/Engine/Source/Developer/SessionFrontend/Private/Models/SessionBrowserOwnerFilter.h
@@ -7,14 +7,10 @@
 #pragma once
 
 
-/**
- * Type definition for shared pointers to instances of FSessionBrowserOwnerFilter.
- */
+/** Type definition for shared pointers to instances of FSessionBrowserOwnerFilter. */
 typedef TSharedPtr<class FSessionBrowserOwnerFilter> FSessionBrowserOwnerFilterPtr;
 
-/**
- * Type definition for shared references to instances of FSessionBrowserOwnerFilter.
- */
+/** Type definition for shared references to instances of FSessionBrowserOwnerFilter. */
 typedef TSharedRef<class FSessionBrowserOwnerFilter> FSessionBrowserOwnerFilterRef;
 
 
@@ -38,9 +34,9 @@ public:
 	/**
 	 * Creates and initializes a new instance.
 	 *
-	 * @param InOwnerName - The user name of the session owner.
-	 * @param InEnabled - Whether this filter is enabled.
-	 * @param InOnStateChanged - A delegate that is executed when the filter's enabled state changed.
+	 * @param InOwnerName The user name of the session owner.
+	 * @param InEnabled Whether this filter is enabled.
+	 * @param InOnStateChanged A delegate that is executed when the filter's enabled state changed.
 	 */
 	FSessionBrowserOwnerFilter( const FString& InOwnerName, bool InEnabled, FOnSessionBrowserOwnerFilterStateChanged InOnStateChanged )
 		: Enabled(InEnabled)
@@ -48,13 +44,12 @@ public:
 		, OnStateChanged(InOnStateChanged)
 	{ }
 
-
 public:
 
 	/**
 	 * Enables or disables the filter based on the specified check box state.
 	 *
-	 * @param CheckState - The check box state.
+	 * @param CheckState The check box state.
 	 */
 	void EnableFromCheckState( ESlateCheckBoxState::Type CheckState )
 	{
@@ -93,7 +88,6 @@ public:
 		return Enabled;
 	}
 
-
 private:
 
 	// Holds a flag indicating whether this filter is enabled.
@@ -102,7 +96,6 @@ private:
 	// Holds the name of the session owner.
 	FString OwnerName;
 
-
 private:
 
 	// Holds a delegate that is executed when the filter's enabled state changed.
diff --git a/Engine/Source/Developer/SessionFrontend/Private/Models/SessionConsoleCategoryFilter.h b/Engine/Source/Developer/SessionFrontend/Private/Models/SessionConsoleCategoryFilter.h
--- a/Engine/Source/Developer/SessionFrontend/Private/Models/SessionConsoleCategoryFilter.h
+++ b/Engine/Source/Developer/SessionFrontend/Private/Models/SessionConsoleCategoryFilter.h
@@ -7,14 +7,10 @@
 #pragma once
 
 
-/**
- * Type definition for shared pointers to instances of FSessionConsoleCategoryFilter.
- */
+/** Type definition for shared pointers to instances of FSessionConsoleCategoryFilter. */
 typedef TSharedPtr<class FSessionConsoleCategoryFilter> FSessionConsoleCategoryFilterPtr;
 
-/**
- * Type definition for shared references to instances of FSessionConsoleCategoryFilter.
- */
+/** Type definition for shared references to instances of FSessionConsoleCategoryFilter. */
 typedef TSharedRef<class FSessionConsoleCategoryFilter> FSessionConsoleCategoryFilterRef;
 
 
@@ -38,9 +34,9 @@ public:
 	/**
 	 * Creates and initializes a new instance.
 	 *
-	 * @param InCategory - The filter category.
-	 * @param InEnabled - Whether this filter is enabled.
-	 * @param InOnStateChanged - A delegate that is executed when the filter's enabled state changed.
+	 * @param InCategory The filter category.
+	 * @param InEnabled Whether this filter is enabled.
+	 * @param InOnStateChanged A delegate that is executed when the filter's enabled state changed.
 	 */
 	FSessionConsoleCategoryFilter( const FName& InCategory, bool InEnabled, FOnSessionConsoleCategoryFilterStateChanged InOnStateChanged )
 		: Category(InCategory)
@@ -48,13 +44,12 @@ public:
 		, OnStateChanged(InOnStateChanged)
 	{ }
 
-
 public:
 
 	/**
 	 * Enables or disables the filter based on the specified check box state.
 	 *
-	 * @param CheckState - The check box state.
+	 * @param CheckState The check box state.
 	 */
 	void EnableFromCheckState( ESlateCheckBoxState::Type CheckState )
 	{
@@ -93,7 +88,6 @@ public:
 		return Enabled;
 	}
 
-
 private:
 
 	// Holds the filter's category.
@@ -102,7 +96,6 @@ private:
 	// Holds a flag indicating whether this filter is enabled.
 	bool Enabled;
 
-
 private:
 
 	// Holds a delegate that is executed when the filter's enabled state changed.
diff --git a/Engine/Source/Developer/SessionFrontend/Private/Models/SessionConsoleLogMessage.h b/Engine/Source/Developer/SessionFrontend/Private/Models/SessionConsoleLogMessage.h
--- a/Engine/Source/Developer/SessionFrontend/Private/Models/SessionConsoleLogMessage.h
+++ b/Engine/Source/Developer/SessionFrontend/Private/Models/SessionConsoleLogMessage.h
@@ -12,41 +12,28 @@
  */
 struct FSessionConsoleLogMessage
 {
-	/**
-	 * Holds the log category.
-	 */
+	/** Holds the log category. */
 	FName Category;
 
-	/**
-	 * Holds the name of the sender's engine instance.
-	 */
+	/** Holds the name of the sender's engine instance. */
 	FString InstanceName;
 
-	/**
-	 * Holds the sender's address.
-	 */
+	/** Holds the sender's address. */
 	FMessageAddress Sender;
 
-	/**
-	 * Holds the message text.
-	 */
+	/** Holds the message text. */
 	FString Text;
 
-	/**
-	 * Holds the time at which the message was generated.
-	 */
+	/** Holds the time at which the message was generated. */
 	FDateTime Time;
 
-	/**
-	 * Holds the number of seconds from the start of the instance at which the message was generated.
-	 */
+	/** Holds the number of seconds from the start of the instance at which the message was generated. */
 	double TimeSeconds;
 
-	/**
-	 * Holds the verbosity type.
-	 */
+	/** Holds the verbosity type. */
 	ELogVerbosity::Type Verbosity;
 
+public:
 
 	/**
 	 * Creates and initializes a new instance.
diff --git a/Engine/Source/Developer/SessionFrontend/Private/Models/SessionConsoleVerbosityFilter.h b/Engine/Source/Developer/SessionFrontend/Private/Models/SessionConsoleVerbosityFilter.h
--- a/Engine/Source/Developer/SessionFrontend/Private/Models/SessionConsoleVerbosityFilter.h
+++ b/Engine/Source/Developer/SessionFrontend/Private/Models/SessionConsoleVerbosityFilter.h
@@ -7,14 +7,10 @@
 #pragma once
 
 
-/**
- * Type definition for shared pointers to instances of FSessionConsoleVerbosityFilter.
- */
+/** Type definition for shared pointers to instances of FSessionConsoleVerbosityFilter. */
 typedef TSharedPtr<class FSessionConsoleVerbosityFilter> FSessionConsoleVerbosityFilterPtr;
 
-/**
- * Type definition for shared references to instances of FSessionConsoleVerbosityFilter.
- */
+/** Type definition for shared references to instances of FSessionConsoleVerbosityFilter. */
 typedef TSharedRef<class FSessionConsoleVerbosityFilter> FSessionConsoleVerbosityFilterRef;
 
 
@@ -38,11 +34,11 @@ public:
 	/**
 	 * Creates and initializes a new instance.
 	 *
-	 * @param InVerbosity - The log verbosity level to filter.
-	 * @param InIcon - The filter's icon.
-	 * @param InEnabled - Whether this filter is enabled.
-	 * @param InName - The filter's name.
-	 * @param InOnStateChanged - A delegate that is executed when the filter's enabled state changed.
+	 * @param InVerbosity The log verbosity level to filter.
+	 * @param InIcon The filter's icon.
+	 * @param InEnabled Whether this filter is enabled.
+	 * @param InName The filter's name.
+	 * @param InOnStateChanged A delegate that is executed when the filter's enabled state changed.
 	 */
 	FSessionConsoleVerbosityFilter( ELogVerbosity::Type InVerbosity, const FSlateBrush* InIcon, bool InEnabled, const FString& InName, FOnSessionConsoleVerbosityFilterStateChanged InOnStateChanged)
 		: Enabled(InEnabled)
@@ -52,13 +48,12 @@ public:
 		, OnStateChanged(InOnStateChanged)
 	{ }
 
-
 public:
 
 	/**
 	 * Enables or disables the filter based on the specified check box state.
 	 *
-	 * @param CheckState - The check box state.
+	 * @param CheckState The check box state.
 	 */
 	void EnableFromCheckState( ESlateCheckBoxState::Type CheckState )
 	{
@@ -131,7 +126,6 @@ private:
 	// Holds the filter's verbosity level.
 	ELogVerbosity::Type Verbosity;
 
-
 private:
 
 	// Holds a delegate that is executed when the filter's enabled state changed.
diff --git a/Engine/Source/Developer/SessionFrontend/Private/Widgets/Browser/SSessionBrowser.h b/Engine/Source/Developer/SessionFrontend/Private/Widgets/Browser/SSessionBrowser.h
--- a/Engine/Source/Developer/SessionFrontend/Private/Widgets/Browser/SSessionBrowser.h
+++ b/Engine/Source/Developer/SessionFrontend/Private/Widgets/Browser/SSessionBrowser.h
@@ -30,8 +30,8 @@ public:
 	/**
 	 * Construct this widget.
 	 *
-	 * @param InArgs - The declaration data for this widget.
-	 * @param InSessionManager - The session manager to use.
+	 * @param InArgs The declaration data for this widget.
+	 * @param InSessionManager The session manager to use.
 	 */
 	void Construct( const FArguments& InArgs, ISessionManagerRef InSessionManager );
 
@@ -45,8 +45,7 @@ protected:
 	/**
 	 * Gets the display name of the specified session.
 	 *
-	 * @param SessionInfo - The session to get the name for.
-	 *
+	 * @param SessionInfo The session to get the name for.
 	 * @param Session name.
 	 */
 	FText GetSessionName( const ISessionInfoPtr& SessionInfo ) const;
@@ -94,7 +93,7 @@ private:
 	TSharedPtr<SSessionInstanceList> InstanceListView;
 
 	// Holds the session combo box.
-	TSharedPtr<SComboBox<ISessionInfoPtr> > SessionComboBox;
+	TSharedPtr<SComboBox<ISessionInfoPtr>> SessionComboBox;
 
 	// Holds the filtered list of sessions to be displayed.
 	TArray<ISessionInfoPtr> SessionList;
@@ -111,5 +110,5 @@ private:
 	void HandleSessionListViewSelectionChanged( ISessionInfoPtr Item, ESelectInfo::Type SelectInfo );
 
 	// Holds the session list view.
-	TSharedPtr<SListView<ISessionInfoPtr> > SessionListView;
+	TSharedPtr<SListView<ISessionInfoPtr>> SessionListView;
 };
diff --git a/Engine/Source/Developer/SessionFrontend/Private/Widgets/Browser/SSessionBrowserCommandBar.h b/Engine/Source/Developer/SessionFrontend/Private/Widgets/Browser/SSessionBrowserCommandBar.h
--- a/Engine/Source/Developer/SessionFrontend/Private/Widgets/Browser/SSessionBrowserCommandBar.h
+++ b/Engine/Source/Developer/SessionFrontend/Private/Widgets/Browser/SSessionBrowserCommandBar.h
@@ -18,13 +18,12 @@ public:
 	SLATE_BEGIN_ARGS(SSessionBrowserCommandBar) { }
 	SLATE_END_ARGS()
 
-
 public:
 
 	/**
 	 * Construct this widget
 	 *
-	 * @param InArgs - The declaration data for this widget.
+	 * @param InArgs The declaration data for this widget.
 	 */
 	void Construct( const FArguments& InArgs );
 };
diff --git a/Engine/Source/Developer/SessionFrontend/Private/Widgets/Browser/SSessionBrowserInstanceList.h b/Engine/Source/Developer/SessionFrontend/Private/Widgets/Browser/SSessionBrowserInstanceList.h
--- a/Engine/Source/Developer/SessionFrontend/Private/Widgets/Browser/SSessionBrowserInstanceList.h
+++ b/Engine/Source/Developer/SessionFrontend/Private/Widgets/Browser/SSessionBrowserInstanceList.h
@@ -41,8 +41,8 @@ public:
 	/**
 	 * Construct this widget.
 	 *
-	 * @param InArgs - The declaration data for this widget.
-	 * @param InSessionManager - The session manager to use.
+	 * @param InArgs The declaration data for this widget.
+	 * @param InSessionManager The session manager to use.
 	 */
 	void Construct( const FArguments& InArgs, ISessionManagerRef InSessionManager )
 	{
@@ -176,7 +176,7 @@ private:
 	TArray<ISessionInstanceInfoPtr> InstanceList;
 
 	// Holds the instance list view.
-	TSharedPtr<SListView<ISessionInstanceInfoPtr> > InstanceListView;
+	TSharedPtr<SListView<ISessionInstanceInfoPtr>> InstanceListView;
 
 	// Holds a reference to the session manager.
 	ISessionManagerPtr SessionManager;
diff --git a/Engine/Source/Developer/SessionFrontend/Private/Widgets/Browser/SSessionBrowserInstanceListRow.h b/Engine/Source/Developer/SessionFrontend/Private/Widgets/Browser/SSessionBrowserInstanceListRow.h
--- a/Engine/Source/Developer/SessionFrontend/Private/Widgets/Browser/SSessionBrowserInstanceListRow.h
+++ b/Engine/Source/Developer/SessionFrontend/Private/Widgets/Browser/SSessionBrowserInstanceListRow.h
@@ -36,7 +36,7 @@ public:
 	/**
 	 * Constructs the widget.
 	 *
-	 * @param InArgs - The construction arguments.
+	 * @param InArgs The construction arguments.
 	 */
 	void Construct( const FArguments& InArgs, const TSharedRef<STableViewBase>& InOwnerTableView )
 	{
@@ -50,8 +50,7 @@ public:
 	/**
 	 * Generates the widget for the specified column.
 	 *
-	 * @param ColumnName - The name of the column to generate the widget for.
-	 *
+	 * @param ColumnName The name of the column to generate the widget for.
 	 * @return The widget.
 	 */
 	virtual TSharedRef<SWidget> GenerateWidgetForColumn( const FName& ColumnName ) OVERRIDE
diff --git a/Engine/Source/Developer/SessionFrontend/Private/Widgets/Browser/SSessionBrowserSessionListRow.h b/Engine/Source/Developer/SessionFrontend/Private/Widgets/Browser/SSessionBrowserSessionListRow.h
--- a/Engine/Source/Developer/SessionFrontend/Private/Widgets/Browser/SSessionBrowserSessionListRow.h
+++ b/Engine/Source/Developer/SessionFrontend/Private/Widgets/Browser/SSessionBrowserSessionListRow.h
@@ -29,7 +29,7 @@ public:
 	/**
 	 * Constructs the widget.
 	 *
-	 * @param InArgs - The construction arguments.
+	 * @param InArgs The construction arguments.
 	 */
 	void Construct( const FArguments& InArgs, const TSharedRef<STableViewBase>& InOwnerTableView )
 	{
@@ -44,8 +44,7 @@ public:
 	/**
 	 * Generates the widget for the specified column.
 	 *
-	 * @param ColumnName - The name of the column to generate the widget for.
-	 *
+	 * @param ColumnName The name of the column to generate the widget for.
 	 * @return The widget.
 	 */
 	virtual TSharedRef<SWidget> GenerateWidgetForColumn( const FName& ColumnName ) OVERRIDE
diff --git a/Engine/Source/Developer/SessionFrontend/Private/Widgets/Console/SSessionConsole.h b/Engine/Source/Developer/SessionFrontend/Private/Widgets/Console/SSessionConsole.h
--- a/Engine/Source/Developer/SessionFrontend/Private/Widgets/Console/SSessionConsole.h
+++ b/Engine/Source/Developer/SessionFrontend/Private/Widgets/Console/SSessionConsole.h
@@ -33,8 +33,8 @@ public:
 	/**
 	 * Construct this widget
 	 *
-	 * @param InArgs - The declaration data for this widget.
-	 * @param InSessionManager - The session manager to use.
+	 * @param InArgs The declaration data for this widget.
+	 * @param InSessionManager The session manager to use.
 	 */
 	void Construct( const FArguments& InArgs, ISessionManagerRef InSessionManager );
 
@@ -155,7 +155,7 @@ private:
 	FString LastLogFileSaveDirectory;
 
 	// Holds the log list view.
- 	TSharedPtr<SListView<FSessionLogMessagePtr> > LogListView;
+ 	TSharedPtr<SListView<FSessionLogMessagePtr>> LogListView;
 
  	// Holds the filtered list of log messages.
  	TArray<FSessionLogMessagePtr> LogMessages;
diff --git a/Engine/Source/Developer/SessionFrontend/Private/Widgets/Console/SSessionConsoleCommandBar.h b/Engine/Source/Developer/SessionFrontend/Private/Widgets/Console/SSessionConsoleCommandBar.h
--- a/Engine/Source/Developer/SessionFrontend/Private/Widgets/Console/SSessionConsoleCommandBar.h
+++ b/Engine/Source/Developer/SessionFrontend/Private/Widgets/Console/SSessionConsoleCommandBar.h
@@ -42,14 +42,14 @@ public:
 	/**
 	 * Construct this widget
 	 *
-	 * @param InArgs - The declaration data for this widget.
+	 * @param InArgs The declaration data for this widget.
 	 */
 	void Construct( const FArguments& InArgs );
 
 	/**
 	 * Sets the number of selected engine instances.
 	 *
-	 * @param Count - Number of selected instances.
+	 * @param Count Number of selected instances.
 	 */
 	void SetNumSelectedInstances( int Count );
 
diff --git a/Engine/Source/Developer/SessionFrontend/Private/Widgets/Console/SSessionConsoleFilterBar.h b/Engine/Source/Developer/SessionFrontend/Private/Widgets/Console/SSessionConsoleFilterBar.h
--- a/Engine/Source/Developer/SessionFrontend/Private/Widgets/Console/SSessionConsoleFilterBar.h
+++ b/Engine/Source/Developer/SessionFrontend/Private/Widgets/Console/SSessionConsoleFilterBar.h
@@ -29,15 +29,14 @@ public:
 	/**
 	 * Construct this widget
 	 *
-	 * @param InArgs - The declaration data for this widget.
+	 * @param InArgs The declaration data for this widget.
 	 */
 	void Construct( const FArguments& InArgs );
 
 	/**
 	 * Filters the specified log message based on the current filter settings.
 	 *
-	 * @param LogMessage - The log message to filter.
-	 *
+	 * @param LogMessage The log message to filter.
 	 * @return true if the log message passed the filter, false otherwise.
 	 */
 	bool FilterLogMessage( const FSessionLogMessageRef& LogMessage );
@@ -62,16 +61,16 @@ protected:
 	/**
 	 * Adds a category filter.
 	 *
-	 * @param Category - The filter's category.
+	 * @param Category The filter's category.
 	 */
 	void AddCategoryFilter( const FName& Category );
 
 	/**
 	 * Adds a verbosity filter.
 	 *
-	 * @param Verbosity - The filters verbosity level.
-	 * @param Name - The name of the filter.
-	 * @param Icon - The name of the filter's icon.
+	 * @param Verbosity The filters verbosity level.
+	 * @param Name The name of the filter.
+	 * @param Icon The name of the filter's icon.
 	 */
 	void AddVerbosityFilter( ELogVerbosity::Type Verbosity, const FString& Name, const FName& Icon );
 
@@ -107,7 +106,7 @@ private:
 	TArray<FSessionConsoleCategoryFilterPtr> CategoriesList;
 
 	// Holds the category filters list view.
-	TSharedPtr<SListView<FSessionConsoleCategoryFilterPtr> > CategoriesListView;
+	TSharedPtr<SListView<FSessionConsoleCategoryFilterPtr>> CategoriesListView;
 
 	// Holds the log message counters for category filters.
 	TMap<FName, int32> CategoryCounters;
@@ -128,7 +127,7 @@ private:
 	TArray<FSessionConsoleVerbosityFilterPtr> VerbositiesList;
 
 	// Holds the verbosity filters list view.
-	TSharedPtr< SListView<FSessionConsoleVerbosityFilterPtr> > VerbositiesListView;
+	TSharedPtr< SListView<FSessionConsoleVerbosityFilterPtr>> VerbositiesListView;
 
 	// Holds the log message counters for verbosity filters.
 	TMap<ELogVerbosity::Type, int32> VerbosityCounters;
diff --git a/Engine/Source/Developer/SessionFrontend/Private/Widgets/Console/SSessionConsoleLogTableRow.h b/Engine/Source/Developer/SessionFrontend/Private/Widgets/Console/SSessionConsoleLogTableRow.h
--- a/Engine/Source/Developer/SessionFrontend/Private/Widgets/Console/SSessionConsoleLogTableRow.h
+++ b/Engine/Source/Developer/SessionFrontend/Private/Widgets/Console/SSessionConsoleLogTableRow.h
@@ -20,13 +20,12 @@ public:
 		SLATE_ARGUMENT(FSessionLogMessagePtr, LogMessage)
 	SLATE_END_ARGS()
 
-
 public:
 
 	/**
 	 * Constructs the widget.
 	 *
-	 * @param InArgs - The construction arguments.
+	 * @param InArgs The construction arguments.
 	 */
 	void Construct( const FArguments& InArgs, const TSharedRef<STableViewBase>& InOwnerTableView )
 	{
@@ -36,9 +35,10 @@ public:
 		SMultiColumnTableRow<FSessionLogMessagePtr>::Construct(FSuperRowType::FArguments(), InOwnerTableView);
 	}
 
-
 public:
 
+	// Begin SMultiColumnTableRow interface
+
 	virtual TSharedRef<SWidget> GenerateWidgetForColumn( const FName& ColumnName ) OVERRIDE
 	{
 		if (ColumnName == "Instance")
@@ -114,6 +114,7 @@ public:
 		return SNullWidget::NullWidget;
 	}
 
+	// End SMultiColumnTableRow interface
 
 private:
 
@@ -141,7 +142,6 @@ private:
 		}
 	}
 
-
 private:
 
 	// Holds the highlight string for the log message.
diff --git a/Engine/Source/Developer/SessionFrontend/Private/Widgets/Console/SSessionConsoleShortcutWindow.cpp b/Engine/Source/Developer/SessionFrontend/Private/Widgets/Console/SSessionConsoleShortcutWindow.cpp
--- a/Engine/Source/Developer/SessionFrontend/Private/Widgets/Console/SSessionConsoleShortcutWindow.cpp
+++ b/Engine/Source/Developer/SessionFrontend/Private/Widgets/Console/SSessionConsoleShortcutWindow.cpp
@@ -13,6 +13,13 @@
 /* SSessionConsoleShortcutWindow interface
  *****************************************************************************/
 
+void SSessionConsoleShortcutWindow::AddShortcut( const FString& InName, const FString& InCommandString )
+{
+	AddShortcutInternal(InName, InCommandString);
+	SaveShortcuts();
+}
+
+
 void SSessionConsoleShortcutWindow::Construct( const FArguments& InArgs )
 {
 	OnCommandSubmitted = InArgs._OnCommandSubmitted;
@@ -20,52 +27,150 @@ void SSessionConsoleShortcutWindow::Construct( const FArguments& InArgs )
 	ChildSlot
 	[
 		SNew(SVerticalBox)
+
 		+ SVerticalBox::Slot()
-		.FillHeight(1.0f)
-		.Padding(6.0f, 0.0f, 0.0f, 0.0f)
-		[
-			SAssignNew(ShortcutListView, SListView< TSharedPtr<FConsoleShortcutData> >)
-			.ItemHeight(24.0f)
-			.ListItemsSource(&Shortcuts)
-			.SelectionMode(ESelectionMode::None)
-			.OnGenerateRow(this, &SSessionConsoleShortcutWindow::OnGenerateWidgetForShortcut)
-			.HeaderRow
-			(
-			SNew(SHeaderRow)
-
-			+ SHeaderRow::Column("Command")
-			.DefaultLabel(LOCTEXT("ShortcutHeaderText", "Shortcuts"))
-			)
-		]
+			.FillHeight(1.0f)
+			.Padding(6.0f, 0.0f, 0.0f, 0.0f)
+			[
+				SAssignNew(ShortcutListView, SListView<TSharedPtr<FConsoleShortcutData>>)
+					.ItemHeight(24.0f)
+					.ListItemsSource(&Shortcuts)
+					.SelectionMode(ESelectionMode::None)
+					.OnGenerateRow(this, &SSessionConsoleShortcutWindow::HandleShortcutListViewGenerateRow)
+					.HeaderRow
+					(
+						SNew(SHeaderRow)
+
+						+ SHeaderRow::Column("Command")
+							.DefaultLabel(LOCTEXT("ShortcutHeaderText", "Shortcuts"))
+					)
+			]
 	];
 
 	LoadShortcuts();
+	RebuildUI();
+}
 
+
+/* SSessionConsoleShortcutWindow implementation
+ *****************************************************************************/
+
+void SSessionConsoleShortcutWindow::AddShortcutInternal( const FString& InName, const FString& InCommandString )
+{
+	TSharedPtr<FConsoleShortcutData> NewCommand = MakeShareable(new FConsoleShortcutData());
+	NewCommand->Name = InName;
+	NewCommand->Command = InCommandString;
+
+	Shortcuts.Add( NewCommand );
 	RebuildUI();
 }
 
 
+void SSessionConsoleShortcutWindow::EditShortcut( TSharedPtr<FConsoleShortcutData> InShortcut, bool bInEditCommand, FString InPromptTitle )
+{
+	FString DefaultString = bInEditCommand ? InShortcut->Command : InShortcut->Name;
+
+	EditedShortcut = InShortcut;
+	bEditCommand = bInEditCommand;
+
+	TSharedRef<STextEntryPopup> TextEntry = SNew(STextEntryPopup)
+		.Label(InPromptTitle)
+		.DefaultText(FText::FromString(*DefaultString))
+		.OnTextCommitted(this, &SSessionConsoleShortcutWindow::HandleShortcutTextEntryCommitted)
+		.SelectAllTextWhenFocused(true)
+		.ClearKeyboardFocusOnCommit(false);
+
+	NameEntryPopupWindow = FSlateApplication::Get().PushMenu(
+		SharedThis(this),
+		TextEntry,
+		FSlateApplication::Get().GetCursorPos(),
+		FPopupTransitionEffect(FPopupTransitionEffect::TypeInPopup)
+	);
+}
+
+
+FString SSessionConsoleShortcutWindow::GetShortcutFilename( ) const
+{
+	FString Filename = FPaths::EngineSavedDir() + TEXT("ConsoleShortcuts.txt");
+	return Filename;
+}
+
 
-void SSessionConsoleShortcutWindow::AddShortcut(const FString& InName, const FString& InCommandString)
+void SSessionConsoleShortcutWindow::LoadShortcuts( )
 {
-	AddShortcutInternal(InName, InCommandString);
+	//clear out list of commands
+	Shortcuts.Empty();
 
-	//make sure we have these around for next time
-	SaveShortcuts();
+	//read file
+	FString Content;
+	FFileHelper::LoadFileToString( Content, *GetShortcutFilename() );
+
+	TSharedPtr<FJsonObject> ShortcutStream;
+	TSharedRef<TJsonReader<>> Reader = TJsonReaderFactory<>::Create( Content );
+	bool bResult = FJsonSerializer::Deserialize( Reader, ShortcutStream );
+
+	if (ShortcutStream.IsValid())
+	{
+		int32 CommandCount = ShortcutStream->GetNumberField(TEXT("Count"));
+		for (int32 i = 0; i < CommandCount; ++i)
+		{
+			FString Name = ShortcutStream->GetStringField(FString::Printf(TEXT("Shortcut.%d.Name"), i));
+			FString Command = ShortcutStream->GetStringField(FString::Printf(TEXT("Shortcut.%d.Command"), i));
+
+			AddShortcut(Name, Command);
+		}
+	}
 }
 
-void SSessionConsoleShortcutWindow::AddShortcutInternal(const FString& InName, const FString& InCommandString)
+
+void SSessionConsoleShortcutWindow::RebuildUI( )
 {
-	TSharedPtr<FConsoleShortcutData> NewCommand = MakeShareable( new FConsoleShortcutData() );
-	NewCommand->Name = InName;
-	NewCommand->Command = InCommandString;
-	Shortcuts.Add( NewCommand );
+	ShortcutListView->RequestListRefresh();
+	SetVisibility(Shortcuts.Num() > 0 ? EVisibility::Visible : EVisibility::Collapsed);
+}
+
 
+void SSessionConsoleShortcutWindow::RemoveShortcut( TSharedPtr<FConsoleShortcutData> InShortcut )
+{
+	Shortcuts.Remove(InShortcut);
 	RebuildUI();
 }
 
 
-TSharedRef<ITableRow> SSessionConsoleShortcutWindow::OnGenerateWidgetForShortcut(TSharedPtr<FConsoleShortcutData> InItem, const TSharedRef<STableViewBase>& OwnerTable )
+void SSessionConsoleShortcutWindow::SaveShortcuts( ) const
+{
+	TSharedPtr<FJsonObject> ShortcutStream = MakeShareable( new FJsonObject );
+
+	ShortcutStream->SetNumberField(TEXT("Count"), Shortcuts.Num());
+	for (int32 i = 0; i < Shortcuts.Num(); ++i)
+	{
+		ShortcutStream->SetStringField(FString::Printf(TEXT("Shortcut.%d.Name"), i), Shortcuts[i]->Name);
+		ShortcutStream->SetStringField(FString::Printf(TEXT("Shortcut.%d.Command"), i), Shortcuts[i]->Command);
+	}
+
+	FString Content;
+	TSharedRef< TJsonWriter<>> Writer = TJsonWriterFactory<>::Create(&Content);
+	FJsonSerializer::Serialize(ShortcutStream.ToSharedRef(), Writer);
+
+	FFileHelper::SaveStringToFile(*Content, *GetShortcutFilename());
+}
+
+
+/* SSessionConsoleShortcutWindow callbacks
+ *****************************************************************************/
+
+FReply SSessionConsoleShortcutWindow::HandleExecuteButtonClicked( TSharedPtr<FConsoleShortcutData> InShortcut )
+{
+	if (OnCommandSubmitted.IsBound())
+	{
+		OnCommandSubmitted.Execute(InShortcut->Command);
+	}
+
+	return FReply::Handled();
+}
+
+
+TSharedRef<ITableRow> SSessionConsoleShortcutWindow::HandleShortcutListViewGenerateRow(TSharedPtr<FConsoleShortcutData> InItem, const TSharedRef<STableViewBase>& OwnerTable)
 {
 	FMenuBuilder ContextMenuBuilder(true, NULL);
 
@@ -75,14 +180,14 @@ TSharedRef<ITableRow> SSessionConsoleShortcutWindow::OnGenerateWidgetForShortcut
 			NSLOCTEXT("SessionFrontend", "ContextMenu.EditName", "Edit Name"),
 			FText(),
 			FSlateIcon(),
-			FUIAction(FExecuteAction::CreateSP( this, &SSessionConsoleShortcutWindow::EditShortcut, InItem, false, LOCTEXT("ShortcutOptionsEditNameTitle", "Name:").ToString() ))
+			FUIAction(FExecuteAction::CreateSP(this, &SSessionConsoleShortcutWindow::EditShortcut, InItem, false, LOCTEXT("ShortcutOptionsEditNameTitle", "Name:").ToString()))
 		);
 
 		ContextMenuBuilder.AddMenuEntry(
 			NSLOCTEXT("SessionFrontend", "ContextMenu.EditCommand", "Edit Command"),
 			FText(),
 			FSlateIcon(),
-			FUIAction(FExecuteAction::CreateSP( this, &SSessionConsoleShortcutWindow::EditShortcut, InItem, true, LOCTEXT("ShortcutOptionsEditCommandTitle", "Command:").ToString() ))
+			FUIAction(FExecuteAction::CreateSP(this, &SSessionConsoleShortcutWindow::EditShortcut, InItem, true, LOCTEXT("ShortcutOptionsEditCommandTitle", "Command:").ToString()))
 		);
 	}
 	ContextMenuBuilder.EndSection();
@@ -93,177 +198,74 @@ TSharedRef<ITableRow> SSessionConsoleShortcutWindow::OnGenerateWidgetForShortcut
 			NSLOCTEXT("SessionFrontend", "ContextMenu.DeleteCommand", "Delete Command"),
 			FText(),
 			FSlateIcon(),
-			FUIAction(FExecuteAction::CreateSP( this, &SSessionConsoleShortcutWindow::RemoveShortcut, InItem ))
+			FUIAction(FExecuteAction::CreateSP(this, &SSessionConsoleShortcutWindow::RemoveShortcut, InItem))
 		);
 	}
 	ContextMenuBuilder.EndSection();
 
-	return SNew(STableRow<TSharedPtr<FConsoleShortcutData> >, OwnerTable)
+	return SNew(STableRow<TSharedPtr<FConsoleShortcutData>>, OwnerTable)
 		.Padding(2)
 		[
 			SNew(SHorizontalBox)
 
-			//display the title in a button
 			+ SHorizontalBox::Slot()
-			.FillWidth(1.0f)
-			.Padding(0)
-			[
-				SNew( SButton )
-				.VAlign( VAlign_Center )
-				.ToolTipText( InItem->Command )
-				.OnClicked( FOnClicked::CreateSP(this, &SSessionConsoleShortcutWindow::ExecuteShortcut, InItem) )
+				.FillWidth(1.0f)
+				.Padding(0)
 				[
-					SNew(STextBlock)
-					.Text(InItem->Name)
+					// execute button
+					SNew(SButton)
+						.VAlign(VAlign_Center)
+						.ToolTipText(InItem->Command)
+						.OnClicked(FOnClicked::CreateSP(this, &SSessionConsoleShortcutWindow::HandleExecuteButtonClicked, InItem))
+						[
+							SNew(STextBlock)
+								.Text(InItem->Name)
+						]
 				]
-			]
 
-			//edit options pulldown
-			+SHorizontalBox::Slot()
-			.AutoWidth()
-			.HAlign(HAlign_Right)
-			.VAlign(VAlign_Center)
-			[
-				SNew( SComboButton )
-				.ButtonStyle( FEditorStyle::Get(), "NoBorder" )
-				.ForegroundColor( FEditorStyle::GetSlateColor("DefaultForeground") )
-				.ContentPadding( FMargin( 6.f, 2.f ) )
-				.MenuContent()
+			+ SHorizontalBox::Slot()
+				.AutoWidth()
+				.HAlign(HAlign_Right)
+				.VAlign(VAlign_Center)
 				[
-					ContextMenuBuilder.MakeWidget()
+					// edit options pull-down
+					SNew(SComboButton)
+						.ButtonStyle(FEditorStyle::Get(), "NoBorder")
+						.ForegroundColor(FEditorStyle::GetSlateColor("DefaultForeground"))
+						.ContentPadding(FMargin(6.f, 2.f))
+						.MenuContent()
+						[
+							ContextMenuBuilder.MakeWidget()
+						]
 				]
-			]
 		];
 }
 
 
-/**
- * Callback for when a shortcut is executed
- */
-FReply SSessionConsoleShortcutWindow::ExecuteShortcut( TSharedPtr<FConsoleShortcutData> InShortcut )
-{
-	if (OnCommandSubmitted.IsBound())
-	{
-		OnCommandSubmitted.Execute(InShortcut->Command);
-	}
-	return FReply::Handled();
-}
-
-void SSessionConsoleShortcutWindow::RemoveShortcut( TSharedPtr<FConsoleShortcutData> InShortcut )
-{
-	Shortcuts.Remove(InShortcut);
-
-	RebuildUI();
-
-}
-
-
-void SSessionConsoleShortcutWindow::EditShortcut( TSharedPtr<FConsoleShortcutData> InShortcut, bool bInEditCommand, FString InPromptTitle )
-{
-	FString DefaultString = bInEditCommand ? InShortcut->Command : InShortcut->Name;
-
-	EditedShortcut = InShortcut;
-	bEditCommand = bInEditCommand;
-
-	TSharedRef<STextEntryPopup> TextEntry = 
-		SNew(STextEntryPopup)
-		.Label(InPromptTitle)
-		.DefaultText( FText::FromString(*DefaultString) )
-		.OnTextCommitted(this, &SSessionConsoleShortcutWindow::OnShortcutEditCommitted)
-		.SelectAllTextWhenFocused(true)
-		.ClearKeyboardFocusOnCommit( false );
-
-	NameEntryPopupWindow = FSlateApplication::Get().PushMenu(
-		SharedThis(this),
-		TextEntry,
-		FSlateApplication::Get().GetCursorPos(),
-		FPopupTransitionEffect( FPopupTransitionEffect::TypeInPopup )
-		);
-}
-
-
-void SSessionConsoleShortcutWindow::OnShortcutEditCommitted(const FText& CommandText, ETextCommit::Type CommitInfo)
+void SSessionConsoleShortcutWindow::HandleShortcutTextEntryCommitted( const FText& CommandText, ETextCommit::Type CommitInfo )
 {
 	if (NameEntryPopupWindow.IsValid())
 	{
 		NameEntryPopupWindow->RequestDestroyWindow();
 
 		int32 IndexOfShortcut = Shortcuts.IndexOfByKey(EditedShortcut);
-		if(EditedShortcut.IsValid() && (IndexOfShortcut!=INDEX_NONE))
+		if (EditedShortcut.IsValid() && (IndexOfShortcut != INDEX_NONE))
 		{
 			//make a new version of the command so the list view knows to refresh
-			TSharedPtr<FConsoleShortcutData> NewCommand = MakeShareable( new FConsoleShortcutData() );
+			TSharedPtr<FConsoleShortcutData> NewCommand = MakeShareable(new FConsoleShortcutData());
 			*NewCommand = *EditedShortcut;
 			Shortcuts[IndexOfShortcut] = NewCommand;
 
 			FString& StringToChange = bEditCommand ? NewCommand->Command : NewCommand->Name;
 			StringToChange = CommandText.ToString();
 		}
+
 		EditedShortcut.Reset();
 	}
 
 	RebuildUI();
-
 	SaveShortcuts();
 }
 
-void SSessionConsoleShortcutWindow::RebuildUI()
-{
-	ShortcutListView->RequestListRefresh();
-	SetVisibility( Shortcuts.Num() > 0 ? EVisibility::Visible : EVisibility::Collapsed);
-}
-
-
-void SSessionConsoleShortcutWindow::LoadShortcuts()
-{
-	//clear out list of commands
-	Shortcuts.Empty();
-
-	//read file
-	FString Content;
-	FFileHelper::LoadFileToString( Content, *GetShortcutFilename() );
-
-	TSharedPtr<FJsonObject> ShortcutStream;
-	TSharedRef< TJsonReader<> > Reader = TJsonReaderFactory<>::Create( Content );
-	bool bResult = FJsonSerializer::Deserialize( Reader, ShortcutStream );
-
-	if (ShortcutStream.IsValid())
-	{
-		int32 CommandCount = ShortcutStream->GetNumberField(TEXT("Count"));
-		for (int32 i = 0; i < CommandCount; ++i)
-		{
-			FString Name = ShortcutStream->GetStringField(FString::Printf(TEXT("Shortcut.%d.Name"), i));
-			FString Command = ShortcutStream->GetStringField(FString::Printf(TEXT("Shortcut.%d.Command"), i));
-			AddShortcut(Name, Command);
-		}
-	}
-}
-
-
-void SSessionConsoleShortcutWindow::SaveShortcuts() const
-{
-	TSharedPtr<FJsonObject> ShortcutStream = MakeShareable( new FJsonObject );
-
-	ShortcutStream->SetNumberField(TEXT("Count"), Shortcuts.Num());
-	for (int32 i = 0; i < Shortcuts.Num(); ++i)
-	{
-		ShortcutStream->SetStringField(FString::Printf(TEXT("Shortcut.%d.Name"), i), Shortcuts[i]->Name);
-		ShortcutStream->SetStringField(FString::Printf(TEXT("Shortcut.%d.Command"), i), Shortcuts[i]->Command);
-	}
-
-	FString Content;
-	TSharedRef< TJsonWriter<> > Writer = TJsonWriterFactory<>::Create( &Content );
-	FJsonSerializer::Serialize( ShortcutStream.ToSharedRef(), Writer );
-
-	FFileHelper::SaveStringToFile( *Content, *GetShortcutFilename() );
-}
-
-
-FString SSessionConsoleShortcutWindow::GetShortcutFilename() const
-{
-	FString Filename = FPaths::EngineSavedDir() + TEXT("ConsoleShortcuts.txt");
-	return Filename;
-}
-
 
 #undef LOCTEXT_NAMESPACE
diff --git a/Engine/Source/Developer/SessionFrontend/Private/Widgets/Console/SSessionConsoleShortcutWindow.h b/Engine/Source/Developer/SessionFrontend/Private/Widgets/Console/SSessionConsoleShortcutWindow.h
--- a/Engine/Source/Developer/SessionFrontend/Private/Widgets/Console/SSessionConsoleShortcutWindow.h
+++ b/Engine/Source/Developer/SessionFrontend/Private/Widgets/Console/SSessionConsoleShortcutWindow.h
@@ -9,17 +9,14 @@
 
 struct FConsoleShortcutData
 {
-	/**
-	 * Name of the shortcut
-	 */
+	/** Name of the shortcut */
 	FString Name;
 
-	/**
-	 * Commandline for the shortcut
-	 */
+	/** Command line for the shortcut */
 	FString Command;
 };
 
+
 /**
  * Implements the console filter bar widget.
  */
@@ -39,39 +36,27 @@ public:
 
 public:
 
-	/**
-	 * Construct this widget
-	 *
-	 * @param InArgs - The declaration data for this widget.
-	 */
-	void Construct( const FArguments& InArgs );
-
-
 	/**
 	 * Adds a new shortcut to the list of commands and saves the list
 	 */
-	void AddShortcut(const FString& InName, const FString& InCommandString);
-
-
-private:
+	void AddShortcut( const FString& InName, const FString& InCommandString );
 
 	/**
-	 * Adds a new shortcut to the list of commands
+	 * Constructs this widget.
+	 *
+	 * @param InArgs The construction arguments.
 	 */
-	void AddShortcutInternal(const FString& InName, const FString& InCommandString);
-
-	/** Generates a row widget for a shortcut */
-	TSharedRef<ITableRow> OnGenerateWidgetForShortcut(TSharedPtr<FConsoleShortcutData> InItem, const TSharedRef<STableViewBase>& OwnerTable );
+	void Construct( const FArguments& InArgs );
 
-	/**
-	 * Callback for when a shortcut is executed
-	 */
-	FReply ExecuteShortcut( TSharedPtr<FConsoleShortcutData> InShortcut);
+private:
 
 	/**
-	 * Callback for when a shortcut should be removed
+	 * Adds a new shortcut to the list of commands.
+	 *
+	 * @param InName The name of the shortcut.
+	 * @param InCommandString The command string.
 	 */
-	void RemoveShortcut( TSharedPtr<FConsoleShortcutData> InShortcut );
+	void AddShortcutInternal( const FString& InName, const FString& InCommandString );
 
 	/**
 	 * Callback for when a shortcut name or command is being edited
@@ -79,47 +64,58 @@ private:
 	void EditShortcut( TSharedPtr<FConsoleShortcutData> InShortcut, bool bInEditCommand, FString InPromptTitle );
 
 	/**
-	 * Callback for committing changes to command
+	 * Gets the filename where the shortcuts are saved/loaded from
 	 */
-	void OnShortcutEditCommitted(const FText& CommandText, ETextCommit::Type CommitInfo);
+	FString GetShortcutFilename( ) const;
 
 	/**
-	 * Rebuild UI or hide of there are no entries
+	 * Loads commands from save file
 	 */
-	void RebuildUI();
+	void LoadShortcuts( );
 
 	/**
-	 * Loads commands from save file
+	 * Rebuild UI or hide of there are no entries
 	 */
-	void LoadShortcuts();
+	void RebuildUI( );
 
 	/**
-	 * Saves commands to save file
+	 * Callback for when a shortcut should be removed
 	 */
-	void SaveShortcuts() const;
+	void RemoveShortcut( TSharedPtr<FConsoleShortcutData> InShortcut );
 
 	/**
-	 * Gets the filename where the shortcuts are saved/loaded from
+	 * Saves commands to save file
 	 */
-	FString GetShortcutFilename() const;
+	void SaveShortcuts( ) const;
 
 private:
 
-	// Holds a delegate that is executed when a command is submitted.
-	FOnSessionConsoleCommandSubmitted OnCommandSubmitted;
-
-	/** List of all commands that are currently supported by shortcuts */
-	TArray< TSharedPtr<FConsoleShortcutData> > Shortcuts;
+	// Callback for when a shortcut is executed
+	FReply HandleExecuteButtonClicked( TSharedPtr<FConsoleShortcutData> InShortcut );
 
-	/** The list view for showing all commands */
-	TSharedPtr< SListView< TSharedPtr<FConsoleShortcutData> > > ShortcutListView;
+	// Generates a row widget for a shortcut.
+	TSharedRef<ITableRow> HandleShortcutListViewGenerateRow( TSharedPtr<FConsoleShortcutData> InItem, const TSharedRef<STableViewBase>& OwnerTable );
 
-	/** Reference to owner of the current popup */
-	TSharedPtr<class SWindow> NameEntryPopupWindow;
+	// Callback for committing changes to command
+	void HandleShortcutTextEntryCommitted( const FText& CommandText, ETextCommit::Type CommitInfo );
 
-	/** Reference to owner of the current popup */
-	TSharedPtr<FConsoleShortcutData> EditedShortcut;
+private:
 
 	/** Whether to edit the name or command */
 	bool bEditCommand;
+
+	/** Reference to owner of the current pop-up */
+	TSharedPtr<FConsoleShortcutData> EditedShortcut;
+
+	/** Reference to owner of the current pop-up */
+	TSharedPtr<class SWindow> NameEntryPopupWindow;
+
+	/** Holds a delegate that is executed when a command is submitted. */
+	FOnSessionConsoleCommandSubmitted OnCommandSubmitted;
+
+	/** List of all commands that are currently supported by shortcuts */
+	TArray<TSharedPtr<FConsoleShortcutData>> Shortcuts;
+
+	/** The list view for showing all commands */
+	TSharedPtr<SListView<TSharedPtr<FConsoleShortcutData>>> ShortcutListView;
 };
diff --git a/Engine/Source/Developer/SessionFrontend/Private/Widgets/Console/SSessionConsoleToolbar.h b/Engine/Source/Developer/SessionFrontend/Private/Widgets/Console/SSessionConsoleToolbar.h
--- a/Engine/Source/Developer/SessionFrontend/Private/Widgets/Console/SSessionConsoleToolbar.h
+++ b/Engine/Source/Developer/SessionFrontend/Private/Widgets/Console/SSessionConsoleToolbar.h
@@ -26,8 +26,8 @@ public:
 	/**
 	 * Constructs the widget.
 	 *
-	 * @param InArgs - The construction arguments.
-	 * @param CommandList - The command list to use.
+	 * @param InArgs The construction arguments.
+	 * @param CommandList The command list to use.
 	 */
 	BEGIN_SLATE_FUNCTION_BUILD_OPTIMIZATION
 	void Construct( const FArguments& InArgs, const TSharedRef<FUICommandList>& CommandList )
diff --git a/Engine/Source/Developer/SessionFrontend/Private/Widgets/SSessionFrontend.h b/Engine/Source/Developer/SessionFrontend/Private/Widgets/SSessionFrontend.h
--- a/Engine/Source/Developer/SessionFrontend/Private/Widgets/SSessionFrontend.h
+++ b/Engine/Source/Developer/SessionFrontend/Private/Widgets/SSessionFrontend.h
@@ -17,28 +17,26 @@ public:
 	SLATE_BEGIN_ARGS(SSessionFrontend) { }
 	SLATE_END_ARGS()
 
-
 public:
 
 	/**
 	 * Constructs the application.
 	 *
-	 * @param InArgs - The Slate argument list.
-	 * @param ConstructUnderMajorTab - The major tab which will contain the session front-end.
-	 * @param ConstructUnderWindow - The window in which this widget is being constructed.
+	 * @param InArgs The Slate argument list.
+	 * @param ConstructUnderMajorTab The major tab which will contain the session front-end.
+	 * @param ConstructUnderWindow The window in which this widget is being constructed.
 	 */
 	void Construct( const FArguments& InArgs, const TSharedRef<SDockTab>& ConstructUnderMajorTab, const TSharedPtr<SWindow>& ConstructUnderWindow );
 
-
 protected:
 
 	/**
 	 * Fills the Window menu with menu items.
 	 *
-	 * @param MenuBuilder - The multi-box builder that should be filled with content for this pull-down menu.
-	 * @param RootMenuGroup - The root menu group.
-	 * @param AppMenuGroup - The application menu group.
-	 * @param TabManager - A Tab Manager from which to populate tab spawner menu items.
+	 * @param MenuBuilder The multi-box builder that should be filled with content for this pull-down menu.
+	 * @param RootMenuGroup The root menu group.
+	 * @param AppMenuGroup The application menu group.
+	 * @param TabManager A Tab Manager from which to populate tab spawner menu items.
 	 */
 	static void FillWindowMenu( FMenuBuilder& MenuBuilder, TSharedRef<FWorkspaceItem> RootMenuGroup, TSharedRef<FWorkspaceItem> AppMenuGroup, const TSharedPtr<FTabManager> TabManager );
 
@@ -47,7 +45,6 @@ protected:
 	 */
 	void InitializeControllers( );
 
-
 private:
 
 	// Callback for handling automation module shutdowns.
@@ -56,7 +53,6 @@ private:
 	// Callback for spawning tabs.
 	TSharedRef<SDockTab> HandleTabManagerSpawnTab( const FSpawnTabArgs& Args, FName TabIdentifier ) const;
 
-
 private:
 
 	// Holds the target device proxy manager.
diff --git a/Engine/Source/Developer/SessionFrontend/Public/Interfaces/ISessionFrontendModule.h b/Engine/Source/Developer/SessionFrontend/Public/Interfaces/ISessionFrontendModule.h
--- a/Engine/Source/Developer/SessionFrontend/Public/Interfaces/ISessionFrontendModule.h
+++ b/Engine/Source/Developer/SessionFrontend/Public/Interfaces/ISessionFrontendModule.h
@@ -18,8 +18,7 @@ public:
 	/**
 	 * Creates a session browser widget.
 	 *
-	 * @param SessionManager - The session manager to use.
-	 *
+	 * @param SessionManager The session manager to use.
 	 * @return The new session browser widget.
 	 */
 	virtual TSharedRef<class SWidget> CreateSessionBrowser( const ISessionManagerRef& SessionManager ) = 0;
@@ -27,11 +26,10 @@ public:
 	/**
 	 * Creates a session console widget.
 	 *
-	 * @param SessionManager - The session manager to use.
+	 * @param SessionManager The session manager to use.
 	 */
 	virtual TSharedRef<class SWidget> CreateSessionConsole( const ISessionManagerRef& SessionManager ) = 0;
 
-
 public:
 
 	/**
