UFE: Fixed non-unity build

diff --git a/Engine/Source/Developer/SessionFrontend/Private/Widgets/SSessionFrontend.cpp b/Engine/Source/Developer/SessionFrontend/Private/Widgets/SSessionFrontend.cpp
--- a/Engine/Source/Developer/SessionFrontend/Private/Widgets/SSessionFrontend.cpp
+++ b/Engine/Source/Developer/SessionFrontend/Private/Widgets/SSessionFrontend.cpp
@@ -2,6 +2,7 @@
 
 #include "SessionFrontendPrivatePCH.h"
 #include "SDockTab.h"
+#include "WorkspaceMenuStructureModule.h"
 
 
 #define LOCTEXT_NAMESPACE "SSessionFrontend"
