SessionFrontend: Showing EngineVersion=0 as 'custom build' (only official releases have version numbers)

diff --git a/Engine/Source/Developer/SessionFrontend/Private/Widgets/Browser/SSessionBrowserInstanceList.h b/Engine/Source/Developer/SessionFrontend/Private/Widgets/Browser/SSessionBrowserInstanceList.h
--- a/Engine/Source/Developer/SessionFrontend/Private/Widgets/Browser/SSessionBrowserInstanceList.h
+++ b/Engine/Source/Developer/SessionFrontend/Private/Widgets/Browser/SSessionBrowserInstanceList.h
@@ -119,7 +119,7 @@ private:
 
 		ToolTipTextBuilder.AppendLineFormat(LOCTEXT("InstanceToolTipBuildDate", "Build Date: {0}"), FText::FromString(Instance->GetBuildDate()));
 		ToolTipTextBuilder.AppendLineFormat(LOCTEXT("InstanceToolTipConsoleBuild", "Console Build: {0}"), Instance->IsConsole() ? LOCTEXT("LabelYes", "Yes") : LOCTEXT("LabelNo", "No"));
-		ToolTipTextBuilder.AppendLineFormat(LOCTEXT("InstanceToolTipEngineVersion", "Engine Version: {0}"), FText::FromString(FString::FromInt(Instance->GetEngineVersion())));
+		ToolTipTextBuilder.AppendLineFormat(LOCTEXT("InstanceToolTipEngineVersion", "Engine Version: {0}"), Instance->GetEngineVersion() == 0 ? LOCTEXT("CustomBuildVersion", "Custom Build") : FText::FromString(FString::FromInt(Instance->GetEngineVersion())));
 		ToolTipTextBuilder.AppendLineFormat(LOCTEXT("InstanceToolTipInstanceId", "Instance ID: {0}"), FText::FromString(Instance->GetInstanceId().ToString(EGuidFormats::DigitsWithHyphensInBraces)));
 
 		return ToolTipTextBuilder.ToText();
