SessionFrontend: include what you use; documentation cleanup pass.

diff --git a/Engine/Source/Developer/SessionFrontend/Private/Models/SessionBrowserOwnerFilter.h b/Engine/Source/Developer/SessionFrontend/Private/Models/SessionBrowserOwnerFilter.h
--- a/Engine/Source/Developer/SessionFrontend/Private/Models/SessionBrowserOwnerFilter.h
+++ b/Engine/Source/Developer/SessionFrontend/Private/Models/SessionBrowserOwnerFilter.h
@@ -3,13 +3,6 @@
 #pragma once
 
 
-/** Type definition for shared pointers to instances of FSessionBrowserOwnerFilter. */
-typedef TSharedPtr<class FSessionBrowserOwnerFilter> FSessionBrowserOwnerFilterPtr;
-
-/** Type definition for shared references to instances of FSessionBrowserOwnerFilter. */
-typedef TSharedRef<class FSessionBrowserOwnerFilter> FSessionBrowserOwnerFilterRef;
-
-
 /**
  * Delegate type for session owner filter state changes.
  *
@@ -59,7 +52,7 @@ public:
 	 *
 	 * @return The user name of the session owner.
 	 */
-	const FString& GetOwnerName( )
+	const FString& GetOwnerName()
 	{
 		return OwnerName;
 	}
@@ -69,7 +62,7 @@ public:
 	 *
 	 * @return Checked if the filter is enabled, Unchecked otherwise.
 	 */
-	ESlateCheckBoxState::Type GetCheckStateFromIsEnabled( ) const
+	ESlateCheckBoxState::Type GetCheckStateFromIsEnabled() const
 	{
 		return (Enabled ? ESlateCheckBoxState::Checked : ESlateCheckBoxState::Unchecked);
 	}
@@ -77,23 +70,30 @@ public:
 	/**
 	 * Checks whether this filter is enabled.
 	 *
-	 * @param true if the filter is enabled, false otherwise.
+	 * @return true if the filter is enabled, false otherwise.
 	 */
-	bool IsEnabled( ) const
+	bool IsEnabled() const
 	{
 		return Enabled;
 	}
 
 private:
 
-	// Holds a flag indicating whether this filter is enabled.
+	/** Holds a flag indicating whether this filter is enabled. */
 	bool Enabled;
 
-	// Holds the name of the session owner.
+	/** Holds the name of the session owner. */
 	FString OwnerName;
 
 private:
 
-	// Holds a delegate that is executed when the filter's enabled state changed.
+	/** Holds a delegate that is executed when the filter's enabled state changed. */
 	FOnSessionBrowserOwnerFilterStateChanged OnStateChanged;
 };
+
+
+/** Type definition for shared pointers to instances of FSessionBrowserOwnerFilter. */
+typedef TSharedPtr<FSessionBrowserOwnerFilter> FSessionBrowserOwnerFilterPtr;
+
+/** Type definition for shared references to instances of FSessionBrowserOwnerFilter. */
+typedef TSharedRef<FSessionBrowserOwnerFilter> FSessionBrowserOwnerFilterRef;
diff --git a/Engine/Source/Developer/SessionFrontend/Private/Models/SessionConsoleCategoryFilter.h b/Engine/Source/Developer/SessionFrontend/Private/Models/SessionConsoleCategoryFilter.h
--- a/Engine/Source/Developer/SessionFrontend/Private/Models/SessionConsoleCategoryFilter.h
+++ b/Engine/Source/Developer/SessionFrontend/Private/Models/SessionConsoleCategoryFilter.h
@@ -3,13 +3,6 @@
 #pragma once
 
 
-/** Type definition for shared pointers to instances of FSessionConsoleCategoryFilter. */
-typedef TSharedPtr<class FSessionConsoleCategoryFilter> FSessionConsoleCategoryFilterPtr;
-
-/** Type definition for shared references to instances of FSessionConsoleCategoryFilter. */
-typedef TSharedRef<class FSessionConsoleCategoryFilter> FSessionConsoleCategoryFilterRef;
-
-
 /**
  * Delegate type for category filter state changes.
  *
@@ -59,7 +52,7 @@ public:
 	 *
 	 * @return The category name.
 	 */
-	FName GetCategory( )
+	FName GetCategory()
 	{
 		return Category;
 	}
@@ -69,7 +62,7 @@ public:
 	 *
 	 * @return Checked if the filter is enabled, Unchecked otherwise.
 	 */
-	ESlateCheckBoxState::Type GetCheckStateFromIsEnabled( ) const
+	ESlateCheckBoxState::Type GetCheckStateFromIsEnabled() const
 	{
 		return (Enabled ? ESlateCheckBoxState::Checked : ESlateCheckBoxState::Unchecked);
 	}
@@ -77,23 +70,30 @@ public:
 	/**
 	 * Checks whether this filter is enabled.
 	 *
-	 * @param true if the filter is enabled, false otherwise.
+	 * @return true if the filter is enabled, false otherwise.
 	 */
-	bool IsEnabled( ) const
+	bool IsEnabled() const
 	{
 		return Enabled;
 	}
 
 private:
 
-	// Holds the filter's category.
+	/** Holds the filter's category. */
 	FName Category;
 
-	// Holds a flag indicating whether this filter is enabled.
+	/** Holds a flag indicating whether this filter is enabled. */
 	bool Enabled;
 
 private:
 
-	// Holds a delegate that is executed when the filter's enabled state changed.
+	/** Holds a delegate that is executed when the filter's enabled state changed. */
 	FOnSessionConsoleCategoryFilterStateChanged OnStateChanged;
 };
+
+
+/** Type definition for shared pointers to instances of FSessionConsoleCategoryFilter. */
+typedef TSharedPtr<FSessionConsoleCategoryFilter> FSessionConsoleCategoryFilterPtr;
+
+/** Type definition for shared references to instances of FSessionConsoleCategoryFilter. */
+typedef TSharedRef<FSessionConsoleCategoryFilter> FSessionConsoleCategoryFilterRef;
diff --git a/Engine/Source/Developer/SessionFrontend/Private/Models/SessionConsoleCommands.h b/Engine/Source/Developer/SessionFrontend/Private/Models/SessionConsoleCommands.h
--- a/Engine/Source/Developer/SessionFrontend/Private/Models/SessionConsoleCommands.h
+++ b/Engine/Source/Developer/SessionFrontend/Private/Models/SessionConsoleCommands.h
@@ -4,17 +4,15 @@
 
 
 /**
- * The device details commands
+ * The device details commands.
  */
 class FSessionConsoleCommands
 	: public TCommands<FSessionConsoleCommands>
 {
 public:
 
-	/**
-	 * Default constructor.
-	 */
-	FSessionConsoleCommands( )
+	/** Default constructor. */
+	FSessionConsoleCommands()
 		: TCommands<FSessionConsoleCommands>(
 			"SessionConsole",
 			NSLOCTEXT("Contexts", "SessionConsole", "Session Console"),
@@ -26,7 +24,7 @@ public:
 
 	// TCommands interface
 
-	virtual void RegisterCommands( ) override
+	virtual void RegisterCommands() override
 	{
 		UI_COMMAND(Clear, "Clear Log", "Clear the log window", EUserInterfaceActionType::Button, FInputGesture());
 		UI_COMMAND(SessionCopy, "Copy", "Copy the selected log messages to the clipboard", EUserInterfaceActionType::Button, FInputGesture());
diff --git a/Engine/Source/Developer/SessionFrontend/Private/Models/SessionConsoleVerbosityFilter.h b/Engine/Source/Developer/SessionFrontend/Private/Models/SessionConsoleVerbosityFilter.h
--- a/Engine/Source/Developer/SessionFrontend/Private/Models/SessionConsoleVerbosityFilter.h
+++ b/Engine/Source/Developer/SessionFrontend/Private/Models/SessionConsoleVerbosityFilter.h
@@ -3,13 +3,6 @@
 #pragma once
 
 
-/** Type definition for shared pointers to instances of FSessionConsoleVerbosityFilter. */
-typedef TSharedPtr<class FSessionConsoleVerbosityFilter> FSessionConsoleVerbosityFilterPtr;
-
-/** Type definition for shared references to instances of FSessionConsoleVerbosityFilter. */
-typedef TSharedRef<class FSessionConsoleVerbosityFilter> FSessionConsoleVerbosityFilterRef;
-
-
 /**
  * Delegate type for verbosity filter state changes.
  *
@@ -63,7 +56,7 @@ public:
 	 *
 	 * @return Checked if the filter is enabled, Unchecked otherwise.
 	 */
-	ESlateCheckBoxState::Type GetCheckStateFromIsEnabled( ) const
+	ESlateCheckBoxState::Type GetCheckStateFromIsEnabled() const
 	{
 		return (Enabled ? ESlateCheckBoxState::Checked : ESlateCheckBoxState::Unchecked);
 	}
@@ -73,7 +66,7 @@ public:
 	 *
 	 * @return Filter icon image.
 	 */
-	const FSlateBrush* GetIcon( ) const
+	const FSlateBrush* GetIcon() const
 	{
 		return Icon;
 	}
@@ -83,7 +76,7 @@ public:
 	 *
 	 * @return The name string.
 	 */
-	const FString& GetName( ) const
+	const FString& GetName() const
 	{
 		return Name;
 	}
@@ -93,7 +86,7 @@ public:
 	 *
 	 * @return Verbosity level.
 	 */
-	ELogVerbosity::Type GetVerbosity( ) const
+	ELogVerbosity::Type GetVerbosity() const
 	{
 		return Verbosity;
 	}
@@ -101,29 +94,36 @@ public:
 	/**
 	 * Checks whether this filter is enabled.
 	 *
-	 * @param true if the filter is enabled, false otherwise.
+	 * @return true if the filter is enabled, false otherwise.
 	 */
-	bool IsEnabled( ) const
+	bool IsEnabled() const
 	{
 		return Enabled;
 	}
 
 private:
 
-	// Holds a flag indicating whether this filter is enabled.
+	/** Holds a flag indicating whether this filter is enabled. */
 	bool Enabled;
 
-	// Holds the icon image.
+	/** Holds the icon image. */
 	const FSlateBrush* Icon;
 
-	// Holds the filter's name.
+	/** Holds the filter's name. */
 	FString Name;
 
-	// Holds the filter's verbosity level.
+	/** Holds the filter's verbosity level. */
 	ELogVerbosity::Type Verbosity;
 
 private:
 
-	// Holds a delegate that is executed when the filter's enabled state changed.
+	/** Holds a delegate that is executed when the filter's enabled state changed. */
 	FOnSessionConsoleVerbosityFilterStateChanged OnStateChanged;
 };
+
+
+/** Type definition for shared pointers to instances of FSessionConsoleVerbosityFilter. */
+typedef TSharedPtr<FSessionConsoleVerbosityFilter> FSessionConsoleVerbosityFilterPtr;
+
+/** Type definition for shared references to instances of FSessionConsoleVerbosityFilter. */
+typedef TSharedRef<FSessionConsoleVerbosityFilter> FSessionConsoleVerbosityFilterRef;
diff --git a/Engine/Source/Developer/SessionFrontend/Private/SessionFrontendModule.cpp b/Engine/Source/Developer/SessionFrontend/Private/SessionFrontendModule.cpp
--- a/Engine/Source/Developer/SessionFrontend/Private/SessionFrontendModule.cpp
+++ b/Engine/Source/Developer/SessionFrontend/Private/SessionFrontendModule.cpp
@@ -1,9 +1,9 @@
 // Copyright 1998-2014 Epic Games, Inc. All Rights Reserved.
 
 #include "SessionFrontendPrivatePCH.h"
-
-#include "WorkspaceMenuStructureModule.h"
+#include "ISessionFrontendModule.h"
 #include "SDockTab.h"
+#include "WorkspaceMenuStructureModule.h"
 
 
 static const FName SessionFrontendTabName("SessionFrontend");
@@ -33,7 +33,7 @@ public:
 
 	// IModuleInterface interface
 	
-	virtual void StartupModule( ) override
+	virtual void StartupModule() override
 	{
 		auto& TabSpawnerEntry = FGlobalTabmanager::Get()->RegisterNomadTabSpawner(SessionFrontendTabName, FOnSpawnTab::CreateRaw(this, &FSessionFrontendModule::SpawnSessionFrontendTab))
 			.SetDisplayName(NSLOCTEXT("FSessionFrontendModule", "FrontendTabTitle", "Session Frontend"))
@@ -47,7 +47,7 @@ public:
 #endif
 	}
 
-	virtual void ShutdownModule( ) override
+	virtual void ShutdownModule() override
 	{
 		FGlobalTabmanager::Get()->UnregisterNomadTabSpawner(SessionFrontendTabName);
 	}
@@ -57,7 +57,7 @@ private:
 	/**
 	 * Creates a new session front-end tab.
 	 *
-	 * @param SpawnTabArgs - The arguments for the tab to spawn.
+	 * @param SpawnTabArgs The arguments for the tab to spawn.
 	 * @return The spawned tab.
 	 */
 	TSharedRef<SDockTab> SpawnSessionFrontendTab( const FSpawnTabArgs& SpawnTabArgs )
diff --git a/Engine/Source/Developer/SessionFrontend/Private/SessionFrontendPrivatePCH.h b/Engine/Source/Developer/SessionFrontend/Private/SessionFrontendPrivatePCH.h
--- a/Engine/Source/Developer/SessionFrontend/Private/SessionFrontendPrivatePCH.h
+++ b/Engine/Source/Developer/SessionFrontend/Private/SessionFrontendPrivatePCH.h
@@ -2,8 +2,6 @@
 
 #pragma once
 
-#include "SessionFrontend.h"
-
 
 /* Dependencies
  *****************************************************************************/
diff --git a/Engine/Source/Developer/SessionFrontend/Private/Widgets/Browser/SSessionBrowser.h b/Engine/Source/Developer/SessionFrontend/Private/Widgets/Browser/SSessionBrowser.h
--- a/Engine/Source/Developer/SessionFrontend/Private/Widgets/Browser/SSessionBrowser.h
+++ b/Engine/Source/Developer/SessionFrontend/Private/Widgets/Browser/SSessionBrowser.h
@@ -16,10 +16,8 @@ public:
 
 public:
 
-	/**
-	 * Destructor.
-	 */
-	~SSessionBrowser( );
+	/** Destructor. */
+	~SSessionBrowser();
 
 public:
 
@@ -33,10 +31,8 @@ public:
 
 protected:
 
-	/**
-	 * Filters the session list.
-	 */
-	void FilterSessions( );
+	/** Filters the session list. */
+	void FilterSessions();
 
 	/**
 	 * Gets the display name of the specified session.
@@ -46,65 +42,63 @@ protected:
 	 */
 	FText GetSessionName( const ISessionInfoPtr& SessionInfo ) const;
 
-	/**
-	 * Reloads the sessions list.
-	 */
-	void ReloadSessions( );
+	/** Reloads the sessions list. */
+	void ReloadSessions();
 
 private:
 
-	// Callback for generating a row in the session combo box.
+	/** Callback for generating a row in the session combo box. */
 	TSharedRef<SWidget> HandleSessionComboBoxGenerateWidget( ISessionInfoPtr SessionInfo ) const;
 
-	// Callback for selection changes in the session combo box.
+	/** Callback for selection changes in the session combo box. */
 	void HandleSessionComboBoxSelectionChanged( ISessionInfoPtr SelectedItem, ESelectInfo::Type SelectInfo );
 
-	// Callback for getting the text in the session combo box.
-	FText HandleSessionComboBoxText( ) const;
+	/** Callback for getting the text in the session combo box. */
+	FText HandleSessionComboBoxText() const;
 
-	// Callback for getting the text of the session details area.
-	FText HandleSessionDetailsText( ) const;
+	/** Callback for getting the text of the session details area. */
+	FText HandleSessionDetailsText() const;
 
-	// Callback for changing the selected session in the session manager.
+	/** Callback for changing the selected session in the session manager. */
 	void HandleSessionManagerSelectedSessionChanged( const ISessionInfoPtr& SelectedSession );
 
-	// Callback for updating the session list in the session manager.
-	void HandleSessionManagerSessionsUpdated( );
+	/** Callback for updating the session list in the session manager. */
+	void HandleSessionManagerSessionsUpdated();
 
-	// Callback for clicking the 'Terminate Session' button.
-	FReply HandleTerminateSessionButtonClicked( );
+	/** Callback for clicking the 'Terminate Session' button. */
+	FReply HandleTerminateSessionButtonClicked();
 
-	// Callback for getting the enabled state of the 'Terminate Session' button.
-	bool HandleTerminateSessionButtonIsEnabled( ) const;
+	/** Callback for getting the enabled state of the 'Terminate Session' button. */
+	bool HandleTerminateSessionButtonIsEnabled() const;
 
 private:
 
-	// Holds an unfiltered list of available sessions.
+	/** Holds an unfiltered list of available sessions. */
 	TArray<ISessionInfoPtr> AvailableSessions;
 
-	// Holds the command bar.
+	/** Holds the command bar. */
 	TSharedPtr<SSessionBrowserCommandBar> CommandBar;
 
-	// Holds the instance list view.
+	/** Holds the instance list view. */
 	TSharedPtr<SSessionInstanceList> InstanceListView;
 
-	// Holds the session combo box.
+	/** Holds the session combo box. */
 	TSharedPtr<SComboBox<ISessionInfoPtr>> SessionComboBox;
 
-	// Holds the filtered list of sessions to be displayed.
+	/** Holds the filtered list of sessions to be displayed. */
 	TArray<ISessionInfoPtr> SessionList;
 
-	// Holds a reference to the session manager.
+	/** Holds a reference to the session manager. */
 	ISessionManagerPtr SessionManager;
 
 private:
 
-	// Callback for creating a row widget in the session list view.
+	/** Callback for creating a row widget in the session list view. */
 	TSharedRef<ITableRow> HandleSessionListViewGenerateRow( ISessionInfoPtr Item, const TSharedRef<STableViewBase>& OwnerTable );
 
-	// Handles changing the selection in the session tree.
+	/** Handles changing the selection in the session tree. */
 	void HandleSessionListViewSelectionChanged( ISessionInfoPtr Item, ESelectInfo::Type SelectInfo );
 
-	// Holds the session list view.
+	/** Holds the session list view. */
 	TSharedPtr<SListView<ISessionInfoPtr>> SessionListView;
 };
diff --git a/Engine/Source/Developer/SessionFrontend/Private/Widgets/Browser/SSessionBrowserInstanceList.h b/Engine/Source/Developer/SessionFrontend/Private/Widgets/Browser/SSessionBrowserInstanceList.h
--- a/Engine/Source/Developer/SessionFrontend/Private/Widgets/Browser/SSessionBrowserInstanceList.h
+++ b/Engine/Source/Developer/SessionFrontend/Private/Widgets/Browser/SSessionBrowserInstanceList.h
@@ -20,10 +20,8 @@ public:
 
 public:
 
-	/**
-	 * Destructor.
-	 */
-	~SSessionInstanceList( )
+	/** Destructor. */
+	~SSessionInstanceList()
 	{
 		if (SessionManager.IsValid())
 		{
@@ -82,10 +80,8 @@ public:
 
 protected:
 
-	/**
-	 * Updates the list of instances.
-	 */
-	void UpdateList( )
+	/** Updates the list of instances. */
+	void UpdateList()
 	{
 		// populate the instance list
 		ISessionInfoPtr SelectedSession = SessionManager->GetSelectedSession();
@@ -112,7 +108,7 @@ protected:
 
 private:
 
-	// Handles getting the tool tip text for an instance.
+	/** Handles getting the tool tip text for an instance. */
 	FText HandleInstanceGetToolTipText( ISessionInstanceInfoPtr Instance ) const
 	{
 		FTextBuilder ToolTipTextBuilder;
@@ -125,7 +121,7 @@ private:
 		return ToolTipTextBuilder.ToText();
 	}
 
-	// Handles creating a row widget for the instance list view.
+	/** Handles creating a row widget for the instance list view. */
 	TSharedRef<ITableRow> HandleInstanceListViewGenerateRow( ISessionInstanceInfoPtr Item, const TSharedRef<STableViewBase>& OwnerTable )
 	{
 		return SNew(SSessionBrowserInstanceListRow, OwnerTable)
@@ -133,7 +129,7 @@ private:
 			.ToolTipText(this, &SSessionInstanceList::HandleInstanceGetToolTipText, Item);
 	}
 
-	// Handles changing the instance list selection.
+	/** Handles changing the instance list selection. */
 	void HandleInstanceListViewSelectionChanged( ISessionInstanceInfoPtr Item, ESelectInfo::Type SelectInfo )
 	{
 		IgnoreInstanceSelectionChanged = true;
@@ -148,8 +144,8 @@ private:
 		IgnoreInstanceSelectionChanged = false;
 	}
 
-	// Handles changing the selection state of an instance.
-	void HandleSessionManagerInstanceSelectionChanged( )
+	/** Handles changing the selection state of an instance. */
+	void HandleSessionManagerInstanceSelectionChanged()
 	{
 		if (!IgnoreInstanceSelectionChanged)
 		{
@@ -157,7 +153,7 @@ private:
 		}
 	}
 
-	// Handles changing the selected session in the session manager.
+	/** Handles changing the selected session in the session manager. */
 	void HandleSessionManagerSelectedSessionChanged( const ISessionInfoPtr& SelectedSession )
 	{
 		UpdateList();
@@ -165,16 +161,16 @@ private:
 
 private:
 
-	// Holds a flag indicating whether instance selection changes should be ignored.
+	/** Holds a flag indicating whether instance selection changes should be ignored. */
 	bool IgnoreInstanceSelectionChanged;
 
-	// Holds the list of instances.
+	/** Holds the list of instances. */
 	TArray<ISessionInstanceInfoPtr> InstanceList;
 
-	// Holds the instance list view.
+	/** Holds the instance list view. */
 	TSharedPtr<SListView<ISessionInstanceInfoPtr>> InstanceListView;
 
-	// Holds a reference to the session manager.
+	/** Holds a reference to the session manager. */
 	ISessionManagerPtr SessionManager;
 };
 
diff --git a/Engine/Source/Developer/SessionFrontend/Private/Widgets/Browser/SSessionBrowserInstanceListRow.h b/Engine/Source/Developer/SessionFrontend/Private/Widgets/Browser/SSessionBrowserInstanceListRow.h
--- a/Engine/Source/Developer/SessionFrontend/Private/Widgets/Browser/SSessionBrowserInstanceListRow.h
+++ b/Engine/Source/Developer/SessionFrontend/Private/Widgets/Browser/SSessionBrowserInstanceListRow.h
@@ -147,20 +147,20 @@ public:
 
 private:
 
-	// Callback for getting the text in the 'Device' column.
-	FText HandleDeviceColumnText( ) const
+	/** Callback for getting the text in the 'Device' column. */
+	FText HandleDeviceColumnText() const
 	{
 		return FText::FromString(InstanceInfo->GetDeviceName());
 	}
 
-	// Callback for getting the border color for this row.
-	FSlateColor HandleInstanceBorderBackgroundColor( ) const
+	/** Callback for getting the border color for this row. */
+	FSlateColor HandleInstanceBorderBackgroundColor() const
 	{
 		return FLinearColor((GetTypeHash(InstanceInfo->GetInstanceId()) & 0xff) * 360.0 / 256.0, 0.8, 0.3, 1.0).HSVToLinearRGB();
 	}
 
-	// Callback for getting the border brush for this row.
-	const FSlateBrush* HandleInstanceBorderBrush( ) const
+	/** Callback for getting the border brush for this row. */
+	const FSlateBrush* HandleInstanceBorderBrush() const
 	{
 		if (FDateTime::UtcNow() - InstanceInfo->GetLastUpdateTime() < FTimespan::FromSeconds(10.0))
 		{
@@ -170,14 +170,14 @@ private:
 		return FEditorStyle::GetBrush("ErrorReporting.EmptyBox");
 	}
 
-	// Callback for getting the instance's current level.
-	FString HandleLevelColumnText( ) const
+	/** Callback for getting the instance's current level. */
+	FString HandleLevelColumnText() const
 	{
 		return InstanceInfo->GetCurrentLevel();
 	}
 
-	// Callback for getting the text in the 'Status' column.
-	FText HandleStatusColumnText( ) const
+	/** Callback for getting the text in the 'Status' column. */
+	FText HandleStatusColumnText() const
 	{
 		if (FDateTime::UtcNow() - InstanceInfo->GetLastUpdateTime() < FTimespan::FromSeconds(10.0))
 		{
@@ -187,8 +187,8 @@ private:
 		return LOCTEXT("StatusTimedOut", "Timed Out");
 	}
 
-	// Callback for getting the foreground text color.
-	FSlateColor HandleTextColorAndOpacity( ) const
+	/** Callback for getting the foreground text color. */
+	FSlateColor HandleTextColorAndOpacity() const
 	{
 		if (FDateTime::UtcNow() - InstanceInfo->GetLastUpdateTime() < FTimespan::FromSeconds(10.0))
 		{
@@ -200,10 +200,10 @@ private:
 
 private:
 
-	// Holds a reference to the instance info that is displayed in this row.
+	/** Holds a reference to the instance info that is displayed in this row. */
 	ISessionInstanceInfoPtr InstanceInfo;
 
-	// Holds a delegate to be invoked when the check box state changed.
+	/** Holds a delegate to be invoked when the check box state changed. */
 	FOnCheckStateChanged OnCheckStateChanged;
 };
 
diff --git a/Engine/Source/Developer/SessionFrontend/Private/Widgets/Browser/SSessionBrowserSessionListRow.h b/Engine/Source/Developer/SessionFrontend/Private/Widgets/Browser/SSessionBrowserSessionListRow.h
--- a/Engine/Source/Developer/SessionFrontend/Private/Widgets/Browser/SSessionBrowserSessionListRow.h
+++ b/Engine/Source/Developer/SessionFrontend/Private/Widgets/Browser/SSessionBrowserSessionListRow.h
@@ -91,8 +91,8 @@ public:
 
 protected:
 
-	// Callback for getting the name of the session.
-	FString HandleGetSessionName( ) const
+	/** Callback for getting the name of the session. */
+	FString HandleGetSessionName() const
 	{
 		const FString& SessionName = SessionInfo->GetSessionName();
 
@@ -103,8 +103,6 @@ protected:
 		}
 
 		// generate name for a standalone session
-
-
 		TArray<ISessionInstanceInfoPtr> Instances;
 		SessionInfo->GetInstances(Instances);
 
@@ -130,10 +128,10 @@ protected:
 
 private:
 
-	// Holds the highlight string for the session name.
+	/** Holds the highlight string for the session name. */
 	TAttribute<FText> HighlightText;
 
-	// Holds a reference to the session info that is displayed in this row.
+	/** Holds a reference to the session info that is displayed in this row. */
 	ISessionInfoPtr SessionInfo;
 };
 
diff --git a/Engine/Source/Developer/SessionFrontend/Private/Widgets/Console/SSessionConsole.h b/Engine/Source/Developer/SessionFrontend/Private/Widgets/Console/SSessionConsole.h
--- a/Engine/Source/Developer/SessionFrontend/Private/Widgets/Console/SSessionConsole.h
+++ b/Engine/Source/Developer/SessionFrontend/Private/Widgets/Console/SSessionConsole.h
@@ -19,10 +19,8 @@ public:
 
 public:
 
-	/**
-	 * Destructor.
-	 */
-	~SSessionConsole( );
+	/** Destructor. */
+	~SSessionConsole();
 
 public:
 
@@ -36,20 +34,14 @@ public:
 
 protected:
 
-	/**
-	 * Binds the device commands on our toolbar.
-	 */
-	void BindCommands( );
+	/** Binds the device commands on our toolbar. */
+	void BindCommands();
 
-	/**
-	 * Clears the log list view.
-	 */
+	/** Clears the log list view. */
 	void ClearLog();
 
-	/**
-	 * Copies the selected log messages to the clipboard.
-	 */
-	void CopyLog( );
+	/** Copies the selected log messages to the clipboard. */
+	void CopyLog();
 
 	/**
 	 * Reloads the log messages for the currently selected engine instances.
@@ -58,15 +50,13 @@ protected:
 	 */
 	void ReloadLog( bool FullyReload );
 
-	/**
-	 * Saves all log messages to a file.
-	 */
-	void SaveLog( );
+	/** Saves all log messages to a file. */
+	void SaveLog();
 
 	/**
 	 * Sends the command entered into the input field.
 	 *
-	 * @param CommandString - The command string to send.
+	 * @param CommandString The command string to send.
 	 */
 	void SendCommand( const FString& CommandString );
 
@@ -76,95 +66,95 @@ protected:
 
 private:
 
-	// Callback for executing the 'Clear' action.
-	void HandleClearActionExecute( );
+	/** Callback for executing the 'Clear' action. */
+	void HandleClearActionExecute();
 
-	// Callback for determining the 'Clear' action can execute.
-	bool HandleClearActionCanExecute( );
+	/** Callback for determining the 'Clear' action can execute. */
+	bool HandleClearActionCanExecute();
 
-	// Callback for executing the 'Copy' action.
-	void HandleCopyActionExecute( );
+	/** Callback for executing the 'Copy' action. */
+	void HandleCopyActionExecute();
 
-	// Callback for determining the 'Copy' action can execute.
-	bool HandleCopyActionCanExecute( );
+	/** Callback for determining the 'Copy' action can execute. */
+	bool HandleCopyActionCanExecute();
 
-	// Callback for executing the 'Save' action.
-	void HandleSaveActionExecute( );
+	/** Callback for executing the 'Save' action. */
+	void HandleSaveActionExecute();
 
-	// Callback for determining the 'Save' action can execute.
-	bool HandleSaveActionCanExecute( );
+	/** Callback for determining the 'Save' action can execute. */
+	bool HandleSaveActionCanExecute();
 
-	// Callback for promoting console command to shortcuts.
+	/** Callback for promoting console command to shortcuts. */
 	void HandleCommandBarPromoteToShortcutClicked( const FString& CommandString );
 
-	// Callback for submitting console commands.
+	/** Callback for submitting console commands. */
 	void HandleCommandSubmitted( const FString& CommandString );
 
-	// Callback for changing the filter settings.
-	void HandleFilterChanged( );
+	/** Callback for changing the filter settings. */
+	void HandleFilterChanged();
 
-	// Callback for scrolling a log item into view.
+	/** Callback for scrolling a log item into view. */
 	void HandleLogListItemScrolledIntoView( FSessionLogMessagePtr Item, const TSharedPtr<ITableRow>& TableRow );
 
-	// Callback for generating a row widget for the log list view.
+	/** Callback for generating a row widget for the log list view. */
 	TSharedRef<ITableRow> HandleLogListGenerateRow( FSessionLogMessagePtr Message, const TSharedRef<STableViewBase>& OwnerTable );
 
-	// Callback for getting the highlight string for log messages.
-	FText HandleLogListGetHighlightText( ) const;
+	/** Callback for getting the highlight string for log messages. */
+	FText HandleLogListGetHighlightText() const;
 
-	// Callback for selecting log messages.
+	/** Callback for selecting log messages. */
 	void HandleLogListSelectionChanged( FSessionLogMessagePtr InItem, ESelectInfo::Type SelectInfo );
 
-	// Callback for getting the enabled state of the console box.
-	bool HandleMainContentIsEnabled( ) const;
+	/** Callback for getting the enabled state of the console box. */
+	bool HandleMainContentIsEnabled() const;
 
-	// Callback for determining the visibility of the 'Select a session' overlay.
-	EVisibility HandleSelectSessionOverlayVisibility( ) const;
+	/** Callback for determining the visibility of the 'Select a session' overlay. */
+	EVisibility HandleSelectSessionOverlayVisibility() const;
 
-	// Callback for changing the engine instance selection.
-	void HandleSessionManagerInstanceSelectionChanged( );
+	/** Callback for changing the engine instance selection. */
+	void HandleSessionManagerInstanceSelectionChanged();
 
-	// Callback for received log entries.
+	/** Callback for received log entries. */
 	void HandleSessionManagerLogReceived( const ISessionInfoRef& Session, const ISessionInstanceInfoRef& Instance, const FSessionLogMessageRef& Message );
 
-	// Callback for changing the selected session.
+	/** Callback for changing the selected session. */
 	void HandleSessionManagerSelectedSessionChanged( const ISessionInfoPtr& SelectedSession );
 
 private:
 
-	// Holds an unfiltered list of available log messages.
+	/** Holds an unfiltered list of available log messages. */
 	TArray<FSessionLogMessagePtr> AvailableLogs;
 
-	// Holds the command bar.
+	/** Holds the command bar. */
 	TSharedPtr<SSessionConsoleCommandBar> CommandBar;
 
-	// Holds the filter bar.
+	/** Holds the filter bar. */
 	TSharedPtr<SSessionConsoleFilterBar> FilterBar;
 
-	// Holds the find bar.
+	/** Holds the find bar. */
 	TSharedPtr<SSearchBox> FindBar;
 
-	// Holds the highlight text
+	/** Holds the highlight text. */
 	FString HighlightText;
 
-	// Holds the directory where the log file was last saved to.
+	/** Holds the directory where the log file was last saved to. */
 	FString LastLogFileSaveDirectory;
 
-	// Holds the log list view.
+	/** Holds the log list view. */
  	TSharedPtr<SListView<FSessionLogMessagePtr>> LogListView;
 
- 	// Holds the filtered list of log messages.
+ 	/** Holds the filtered list of log messages. */
  	TArray<FSessionLogMessagePtr> LogMessages;
 
-	// Holds the session manager.
+	/** Holds the session manager. */
 	ISessionManagerPtr SessionManager;
 
-	//Holds the shortcut window
+	/** Holds the shortcut window. */
 	TSharedPtr<SSessionConsoleShortcutWindow> ShortcutWindow;
 
-	// Holds a flag indicating whether the log list should auto-scroll to the last item.
+	/** Holds a flag indicating whether the log list should auto-scroll to the last item. */
 	bool ShouldScrollToLast;
 
-	// The command list for controlling the device
+	/** The command list for controlling the device */
 	TSharedPtr<FUICommandList> UICommandList;
 };
diff --git a/Engine/Source/Developer/SessionFrontend/Private/Widgets/Console/SSessionConsoleCommandBar.h b/Engine/Source/Developer/SessionFrontend/Private/Widgets/Console/SSessionConsoleCommandBar.h
--- a/Engine/Source/Developer/SessionFrontend/Private/Widgets/Console/SSessionConsoleCommandBar.h
+++ b/Engine/Source/Developer/SessionFrontend/Private/Widgets/Console/SSessionConsoleCommandBar.h
@@ -49,48 +49,50 @@ protected:
 
 	/**
 	 * Submits the entered command.
+	 *
+	 * @param Command the command to submit.
 	 */
 	void SubmitCommand( const FString& Command );
 
 private:
 
-	// Handles changing the input text box's content.
+	/** Handles changing the input text box's content. */
 	void HandleInputTextChanged( const FText& InText );
 
-	// Handles committing the input text box's content.
+	/** Handles committing the input text box's content. */
 	void HandleInputTextCommitted( const FText& InText, ETextCommit::Type CommitInfo );
 
-	// Handles showing a history in the input text box.
+	/** Handles showing a history in the input text box. */
 	void HandleInputTextShowingHistory( TArray<FString>& OutHistory );
 
-	// Handles showing suggestions in the input text box.
+	/** Handles showing suggestions in the input text box. */
 	void HandleInputTextShowingSuggestions( const FString& Text, TArray<FString>& OutSuggestions );
 
-	//Handles clicking the promote to shortcut button
-	FReply HandlePromoteToShortcutButtonClicked( );
+	/** Handles clicking the promote to shortcut button. */
+	FReply HandlePromoteToShortcutButtonClicked();
 
-	// Handles clicking the send button.
-	FReply HandleSendButtonClicked( );
+	/** Handles clicking the send button. */
+	FReply HandleSendButtonClicked();
 
 private:
 
-	// Holds the command history.
+	/** Holds the command history. */
 	TArray<FString> CommandHistory;
 
-	// Holds the input text box.
+	/** Holds the input text box. */
 	TSharedPtr<SSuggestionTextBox> InputTextBox;
 
-	// Holds the send button.
+	/** Holds the send button. */
 	TSharedPtr<SButton> SendButton;
 
-	// Holds the promote to shortcut button.
+	/** Holds the promote to shortcut button. */
 	TSharedPtr<SButton> PromoteToShortcutButton;
 
 private:
 
-	// Holds a delegate that is executed when a command is submitted.
+	/** Holds a delegate that is executed when a command is submitted. */
 	FOnSessionConsoleCommandSubmitted OnCommandSubmitted;
 
-	// Holds a delegate that is executed when the Promote To Shortcut button is clicked.
+	/** Holds a delegate that is executed when the Promote To Shortcut button is clicked. */
 	FOnSessionConsoleCommandSubmitted OnPromoteToShortcutClicked;
 };
diff --git a/Engine/Source/Developer/SessionFrontend/Private/Widgets/Console/SSessionConsoleFilterBar.h b/Engine/Source/Developer/SessionFrontend/Private/Widgets/Console/SSessionConsoleFilterBar.h
--- a/Engine/Source/Developer/SessionFrontend/Private/Widgets/Console/SSessionConsoleFilterBar.h
+++ b/Engine/Source/Developer/SessionFrontend/Private/Widgets/Console/SSessionConsoleFilterBar.h
@@ -40,12 +40,10 @@ public:
 	 *
 	 * @return Filter string.
 	 */
-	FText GetFilterText( ) const;
+	FText GetFilterText() const;
 
-	/**
-	 * Resets the categories and filter counters.
-	 */
-	void ResetFilter( );
+	/** Resets the categories and filter counters. */
+	void ResetFilter();
 
 protected:
 
@@ -53,6 +51,7 @@ protected:
 	 * Adds a category filter.
 	 *
 	 * @param Category The filter's category.
+	 * @see AddVerbosityFilter
 	 */
 	void AddCategoryFilter( const FName& Category );
 
@@ -62,69 +61,70 @@ protected:
 	 * @param Verbosity The filters verbosity level.
 	 * @param Name The name of the filter.
 	 * @param Icon The name of the filter's icon.
+	 * @see AddVerbosityFilter
 	 */
 	void AddVerbosityFilter( ELogVerbosity::Type Verbosity, const FString& Name, const FName& Icon );
 
 private:
 
-	// Callback for generating a row widget for the category filter list.
+	/** Callback for generating a row widget for the category filter list. */
 	TSharedRef<ITableRow> HandleCategoryFilterGenerateRow( FSessionConsoleCategoryFilterPtr Filter, const TSharedRef<STableViewBase>& OwnerTable );
 
-	// Callback for getting the text for a row in the category filter drop-down.
+	/** Callback for getting the text for a row in the category filter drop-down. */
 	FString HandleCategoryFilterGetRowText( FSessionConsoleCategoryFilterPtr Filter ) const;
 
-	// Callback for changing the enabled state of a category filter.
+	/** Callback for changing the enabled state of a category filter. */
 	void HandleCategoryFilterStateChanged( const FName& ChangedCategory, bool Enabled );
 
-	// Callback for changing the filter string text box text.
+	/** Callback for changing the filter string text box text. */
 	void HandleFilterStringTextChanged( const FText& NewText );
 
-	// Callback for changing the checked state of the 'Filter' check box.
+	/** Callback for changing the checked state of the 'Filter' check box. */
 	void HandleHighlightOnlyCheckBoxCheckStateChanged( ESlateCheckBoxState::Type CheckedState );
 
-	// Callback for generating a row widget for the verbosity filter list.
+	/** Callback for generating a row widget for the verbosity filter list. */
 	TSharedRef<ITableRow> HandleVerbosityFilterGenerateRow( FSessionConsoleVerbosityFilterPtr Filter, const TSharedRef<STableViewBase>& OwnerTable );
 
-	// Callback for getting the text for a row in the verbosity filter drop-down.
+	/** Callback for getting the text for a row in the verbosity filter drop-down. */
 	FString HandleVerbosityFilterGetRowText( FSessionConsoleVerbosityFilterPtr Filter ) const;
 
-	// Callback for changing the check state of a filter button.
+	/** Callback for changing the check state of a filter button. */
 	void HandleVerbosityFilterStateChanged( ELogVerbosity::Type Verbosity, bool Enabled );
 
 private:
 
-	// Holds the list of category filters.
+	/** Holds the list of category filters. */
 	TArray<FSessionConsoleCategoryFilterPtr> CategoriesList;
 
-	// Holds the category filters list view.
+	/** Holds the category filters list view. */
 	TSharedPtr<SListView<FSessionConsoleCategoryFilterPtr>> CategoriesListView;
 
-	// Holds the log message counters for category filters.
+	/** Holds the log message counters for category filters. */
 	TMap<FName, int32> CategoryCounters;
 
-	// Holds the list of disabled log categories.
+	/** Holds the list of disabled log categories. */
 	TArray<FName> DisabledCategories;
 
-	// Holds the list of disabled log verbosities.
+	/** Holds the list of disabled log verbosities. */
 	TArray<ELogVerbosity::Type> DisabledVerbosities;
 
-	// Holds the filter check box.
+	/** Holds the filter check box. */
 	TSharedPtr<SCheckBox> HighlightOnlyCheckBox;
 
-	// Holds the filter string text box.
+	/** Holds the filter string text box. */
 	TSharedPtr<SSearchBox> FilterStringTextBox;
 
-	// Holds the verbosity filters.
+	/** Holds the verbosity filters. */
 	TArray<FSessionConsoleVerbosityFilterPtr> VerbositiesList;
 
-	// Holds the verbosity filters list view.
+	/** Holds the verbosity filters list view. */
 	TSharedPtr< SListView<FSessionConsoleVerbosityFilterPtr>> VerbositiesListView;
 
-	// Holds the log message counters for verbosity filters.
+	/** Holds the log message counters for verbosity filters. */
 	TMap<ELogVerbosity::Type, int32> VerbosityCounters;
 
 private:
 
-	// Holds a delegate that is executed when the filter settings changed.
+	/** Holds a delegate that is executed when the filter settings changed. */
 	FSimpleDelegate OnFilterChanged;
 };
diff --git a/Engine/Source/Developer/SessionFrontend/Private/Widgets/Console/SSessionConsoleLogTableRow.h b/Engine/Source/Developer/SessionFrontend/Private/Widgets/Console/SSessionConsoleLogTableRow.h
--- a/Engine/Source/Developer/SessionFrontend/Private/Widgets/Console/SSessionConsoleLogTableRow.h
+++ b/Engine/Source/Developer/SessionFrontend/Private/Widgets/Console/SSessionConsoleLogTableRow.h
@@ -112,14 +112,14 @@ public:
 
 private:
 
-	// Gets the border color for this row.
-	FSlateColor HandleGetBorderColor( ) const
+	/** Gets the border color for this row. */
+	FSlateColor HandleGetBorderColor() const
 	{
 		return FLinearColor((GetTypeHash(LogMessage->InstanceId) & 0xff) * 360.0f / 256.0f, 0.8f, 0.3f, 1.0f).HSVToLinearRGB();
 	}
 
-	// Gets the text color for this log entry.
-	FSlateColor HandleGetTextColor( ) const
+	/** Gets the text color for this log entry. */
+	FSlateColor HandleGetTextColor() const
 	{
 		if ((LogMessage->Verbosity == ELogVerbosity::Error) ||
 			(LogMessage->Verbosity == ELogVerbosity::Fatal))
@@ -138,9 +138,9 @@ private:
 
 private:
 
-	// Holds the highlight string for the log message.
+	/** Holds the highlight string for the log message. */
 	TAttribute<FText> HighlightText;
 
-	// Holds a reference to the log message that is displayed in this row.
+	/** Holds a reference to the log message that is displayed in this row. */
 	FSessionLogMessagePtr LogMessage;
 };
diff --git a/Engine/Source/Developer/SessionFrontend/Private/Widgets/Console/SSessionConsoleShortcutWindow.cpp b/Engine/Source/Developer/SessionFrontend/Private/Widgets/Console/SSessionConsoleShortcutWindow.cpp
--- a/Engine/Source/Developer/SessionFrontend/Private/Widgets/Console/SSessionConsoleShortcutWindow.cpp
+++ b/Engine/Source/Developer/SessionFrontend/Private/Widgets/Console/SSessionConsoleShortcutWindow.cpp
@@ -63,7 +63,7 @@ void SSessionConsoleShortcutWindow::AddShortcutInternal( const FString& InName,
 }
 
 
-void SSessionConsoleShortcutWindow::EditShortcut( TSharedPtr<FConsoleShortcutData> InShortcut, bool bInEditCommand, FString InPromptTitle )
+void SSessionConsoleShortcutWindow::HandleEditCommandActionExecute( TSharedPtr<FConsoleShortcutData> InShortcut, bool bInEditCommand, FString InPromptTitle )
 {
 	FString DefaultString = bInEditCommand ? InShortcut->Command : InShortcut->Name;
 
@@ -86,21 +86,21 @@ void SSessionConsoleShortcutWindow::EditShortcut( TSharedPtr<FConsoleShortcutDat
 }
 
 
-FString SSessionConsoleShortcutWindow::GetShortcutFilename( ) const
+FString SSessionConsoleShortcutWindow::GetShortcutFilename() const
 {
 	FString Filename = FPaths::EngineSavedDir() + TEXT("ConsoleShortcuts.txt");
 	return Filename;
 }
 
 
-void SSessionConsoleShortcutWindow::LoadShortcuts( )
+void SSessionConsoleShortcutWindow::LoadShortcuts()
 {
 	//clear out list of commands
 	Shortcuts.Empty();
 
 	//read file
 	FString Content;
-	FFileHelper::LoadFileToString( Content, *GetShortcutFilename() );
+	FFileHelper::LoadFileToString(Content, *GetShortcutFilename());
 
 	TSharedPtr<FJsonObject> ShortcutStream;
 	TSharedRef<TJsonReader<>> Reader = TJsonReaderFactory<>::Create( Content );
@@ -120,21 +120,21 @@ void SSessionConsoleShortcutWindow::LoadShortcuts( )
 }
 
 
-void SSessionConsoleShortcutWindow::RebuildUI( )
+void SSessionConsoleShortcutWindow::RebuildUI()
 {
 	ShortcutListView->RequestListRefresh();
 	SetVisibility(Shortcuts.Num() > 0 ? EVisibility::Visible : EVisibility::Collapsed);
 }
 
 
-void SSessionConsoleShortcutWindow::RemoveShortcut( TSharedPtr<FConsoleShortcutData> InShortcut )
+void SSessionConsoleShortcutWindow::HandleDeleteCommandActionExecute( TSharedPtr<FConsoleShortcutData> InShortcut )
 {
 	Shortcuts.Remove(InShortcut);
 	RebuildUI();
 }
 
 
-void SSessionConsoleShortcutWindow::SaveShortcuts( ) const
+void SSessionConsoleShortcutWindow::SaveShortcuts() const
 {
 	TSharedPtr<FJsonObject> ShortcutStream = MakeShareable( new FJsonObject );
 
@@ -177,14 +177,14 @@ TSharedRef<ITableRow> SSessionConsoleShortcutWindow::HandleShortcutListViewGener
 			NSLOCTEXT("SessionFrontend", "ContextMenu.EditName", "Edit Name"),
 			FText(),
 			FSlateIcon(),
-			FUIAction(FExecuteAction::CreateSP(this, &SSessionConsoleShortcutWindow::EditShortcut, InItem, false, LOCTEXT("ShortcutOptionsEditNameTitle", "Name:").ToString()))
+			FUIAction(FExecuteAction::CreateSP(this, &SSessionConsoleShortcutWindow::HandleEditCommandActionExecute, InItem, false, LOCTEXT("ShortcutOptionsEditNameTitle", "Name:").ToString()))
 		);
 
 		ContextMenuBuilder.AddMenuEntry(
 			NSLOCTEXT("SessionFrontend", "ContextMenu.EditCommand", "Edit Command"),
 			FText(),
 			FSlateIcon(),
-			FUIAction(FExecuteAction::CreateSP(this, &SSessionConsoleShortcutWindow::EditShortcut, InItem, true, LOCTEXT("ShortcutOptionsEditCommandTitle", "Command:").ToString()))
+			FUIAction(FExecuteAction::CreateSP(this, &SSessionConsoleShortcutWindow::HandleEditCommandActionExecute, InItem, true, LOCTEXT("ShortcutOptionsEditCommandTitle", "Command:").ToString()))
 		);
 	}
 	ContextMenuBuilder.EndSection();
@@ -195,7 +195,7 @@ TSharedRef<ITableRow> SSessionConsoleShortcutWindow::HandleShortcutListViewGener
 			NSLOCTEXT("SessionFrontend", "ContextMenu.DeleteCommand", "Delete Command"),
 			FText(),
 			FSlateIcon(),
-			FUIAction(FExecuteAction::CreateSP(this, &SSessionConsoleShortcutWindow::RemoveShortcut, InItem))
+			FUIAction(FExecuteAction::CreateSP(this, &SSessionConsoleShortcutWindow::HandleDeleteCommandActionExecute, InItem))
 		);
 	}
 	ContextMenuBuilder.EndSection();
diff --git a/Engine/Source/Developer/SessionFrontend/Private/Widgets/Console/SSessionConsoleShortcutWindow.h b/Engine/Source/Developer/SessionFrontend/Private/Widgets/Console/SSessionConsoleShortcutWindow.h
--- a/Engine/Source/Developer/SessionFrontend/Private/Widgets/Console/SSessionConsoleShortcutWindow.h
+++ b/Engine/Source/Developer/SessionFrontend/Private/Widgets/Console/SSessionConsoleShortcutWindow.h
@@ -31,7 +31,10 @@ public:
 public:
 
 	/**
-	 * Adds a new shortcut to the list of commands and saves the list
+	 * Adds a new shortcut to the list of commands and saves the list.
+	 *
+	 * @param InName The name of the shortcut to add.
+	 * @param InCommandString The command associated with the shortcut.
 	 */
 	void AddShortcut( const FString& InName, const FString& InCommandString );
 
@@ -42,7 +45,7 @@ public:
 	 */
 	void Construct( const FArguments& InArgs );
 
-private:
+protected:
 
 	/**
 	 * Adds a new shortcut to the list of commands.
@@ -53,63 +56,63 @@ private:
 	void AddShortcutInternal( const FString& InName, const FString& InCommandString );
 
 	/**
-	 * Callback for when a shortcut name or command is being edited
-	 */
-	void EditShortcut( TSharedPtr<FConsoleShortcutData> InShortcut, bool bInEditCommand, FString InPromptTitle );
-
-	/**
-	 * Gets the filename where the shortcuts are saved/loaded from
-	 */
-	FString GetShortcutFilename( ) const;
-
-	/**
-	 * Loads commands from save file
+	 * Gets the name of the file that stores the shortcuts.
+	 *
+	 * @return The file name.
 	 */
-	void LoadShortcuts( );
+	FString GetShortcutFilename() const;
 
 	/**
-	 * Rebuild UI or hide of there are no entries
+	 * Loads commands from save file.
+	 *
+	 * @see SaveShortcuts
 	 */
-	void RebuildUI( );
+	void LoadShortcuts();
 
-	/**
-	 * Callback for when a shortcut should be removed
-	 */
-	void RemoveShortcut( TSharedPtr<FConsoleShortcutData> InShortcut );
+	/** Rebuild UI or hide of there are no entries. */
+	void RebuildUI();
 
 	/**
-	 * Saves commands to save file
+	 * Saves commands to save file.
+	 *
+	 * @see LoadShortcuts
 	 */
-	void SaveShortcuts( ) const;
+	void SaveShortcuts() const;
 
 private:
 
-	// Callback for when a shortcut is executed
+	/** Callback for when a shortcut should be removed. */
+	void HandleDeleteCommandActionExecute( TSharedPtr<FConsoleShortcutData> InShortcut );
+
+	/** Callback for when a shortcut is executed. */
 	FReply HandleExecuteButtonClicked( TSharedPtr<FConsoleShortcutData> InShortcut );
 
-	// Generates a row widget for a shortcut.
+	/** Callback for when a shortcut name or command is being edited. */
+	void HandleEditCommandActionExecute( TSharedPtr<FConsoleShortcutData> InShortcut, bool bInEditCommand, FString InPromptTitle );
+
+	/** Generates a row widget for a shortcut. */
 	TSharedRef<ITableRow> HandleShortcutListViewGenerateRow( TSharedPtr<FConsoleShortcutData> InItem, const TSharedRef<STableViewBase>& OwnerTable );
 
-	// Callback for committing changes to command
+	/** Callback for committing changes to command. */
 	void HandleShortcutTextEntryCommitted( const FText& CommandText, ETextCommit::Type CommitInfo );
 
 private:
 
-	/** Whether to edit the name or command */
+	/** Whether to edit the name or command. */
 	bool bEditCommand;
 
-	/** Reference to owner of the current pop-up */
+	/** Reference to owner of the current pop-up. */
 	TSharedPtr<FConsoleShortcutData> EditedShortcut;
 
-	/** Reference to owner of the current pop-up */
+	/** Reference to owner of the current pop-up. */
 	TSharedPtr<class SWindow> NameEntryPopupWindow;
 
 	/** Holds a delegate that is executed when a command is submitted. */
 	FOnSessionConsoleCommandSubmitted OnCommandSubmitted;
 
-	/** List of all commands that are currently supported by shortcuts */
+	/** List of all commands that are currently supported by shortcuts. */
 	TArray<TSharedPtr<FConsoleShortcutData>> Shortcuts;
 
-	/** The list view for showing all commands */
+	/** The list view for showing all commands. */
 	TSharedPtr<SListView<TSharedPtr<FConsoleShortcutData>>> ShortcutListView;
 };
diff --git a/Engine/Source/Developer/SessionFrontend/Private/Widgets/SSessionFrontend.h b/Engine/Source/Developer/SessionFrontend/Private/Widgets/SSessionFrontend.h
--- a/Engine/Source/Developer/SessionFrontend/Private/Widgets/SSessionFrontend.h
+++ b/Engine/Source/Developer/SessionFrontend/Private/Widgets/SSessionFrontend.h
@@ -35,36 +35,34 @@ protected:
 	 */
 	static void FillWindowMenu( FMenuBuilder& MenuBuilder, const TSharedPtr<FTabManager> TabManager );
 
-	/**
-	 * Creates and initializes the controller classes.
-	 */
-	void InitializeControllers( );
+	/** Creates and initializes the controller classes. */
+	void InitializeControllers();
 
 private:
 
-	// Callback for handling automation module shutdowns.
+	/** Callback for handling automation module shutdowns. */
 	void HandleAutomationModuleShutdown();
 
-	// Callback for spawning tabs.
+	/** Callback for spawning tabs. */
 	TSharedRef<SDockTab> HandleTabManagerSpawnTab( const FSpawnTabArgs& Args, FName TabIdentifier ) const;
 
 private:
 
-	// Holds the target device proxy manager.
+	/** Holds the target device proxy manager. */
 	ITargetDeviceProxyManagerPtr DeviceProxyManager;
 	
-	// Holds a flag indicating whether the launcher overlay is visible.
+	/** Holds a flag indicating whether the launcher overlay is visible. */
 	bool LauncherOverlayVisible;
 
-	// Holds the 'new session' button.
+	/** Holds the 'new session' button. */
 	TSharedPtr<SButton> NewSessionButton;
 
-	// Holds a pointer to the session manager.
+	/** Holds a pointer to the session manager. */
 	ISessionManagerPtr SessionManager;
 
-	// Holds a pointer to the session manager.
+	/** Holds a pointer to the session manager. */
 	IScreenShotManagerPtr ScreenShotManager;
 
-	// Holds the tab manager that manages the front-end's tabs.
+	/** Holds the tab manager that manages the front-end's tabs. */
 	TSharedPtr<FTabManager> TabManager;
 };
diff --git a/Engine/Source/Developer/SessionFrontend/Public/ISessionFrontendModule.h b/Engine/Source/Developer/SessionFrontend/Public/ISessionFrontendModule.h
new file mode 100644
--- /dev/null
+++ b/Engine/Source/Developer/SessionFrontend/Public/ISessionFrontendModule.h
@@ -0,0 +1,40 @@
+// Copyright 1998-2014 Epic Games, Inc. All Rights Reserved.
+
+#pragma once
+
+#include "ModuleInterface.h"
+
+
+// forward declarations
+class ISessionManager;
+class SWidget;
+
+
+/**
+ * Interface for launcher UI modules.
+ */
+class ISessionFrontendModule
+	: public IModuleInterface
+{
+public:
+
+	/**
+	 * Creates a session browser widget.
+	 *
+	 * @param SessionManager The session manager to use.
+	 * @return The new session browser widget.
+	 */
+	virtual TSharedRef<SWidget> CreateSessionBrowser( const TSharedRef<ISessionManager>& SessionManager ) = 0;
+	
+	/**
+	 * Creates a session console widget.
+	 *
+	 * @param SessionManager The session manager to use.
+	 */
+	virtual TSharedRef<SWidget> CreateSessionConsole( const TSharedRef<ISessionManager>& SessionManager ) = 0;
+
+public:
+
+	/** Virtual destructor. */
+	virtual ~ISessionFrontendModule() { }
+};
diff --git a/Engine/Source/Developer/SessionFrontend/Public/Interfaces/ISessionFrontendModule.h b/Engine/Source/Developer/SessionFrontend/Public/Interfaces/ISessionFrontendModule.h
deleted file mode 100644
diff --git a/Engine/Source/Developer/SessionFrontend/Public/SessionFrontend.h b/Engine/Source/Developer/SessionFrontend/Public/SessionFrontend.h
deleted file mode 100644
