Fixed code relying on SLATE_TEXT_ATTRIBUTE for SExpandableArea

Made sure everything was using FText rather than FString.

diff --git a/Engine/Source/Developer/SessionFrontend/Private/Widgets/Console/SSessionConsole.cpp b/Engine/Source/Developer/SessionFrontend/Private/Widgets/Console/SSessionConsole.cpp
--- a/Engine/Source/Developer/SessionFrontend/Private/Widgets/Console/SSessionConsole.cpp
+++ b/Engine/Source/Developer/SessionFrontend/Private/Widgets/Console/SSessionConsole.cpp
@@ -55,7 +55,7 @@ void SSessionConsole::Construct( const FArguments& InArgs, ISessionManagerRef In
 					[
 						// filter bar
 						SNew(SExpandableArea)
-							.AreaTitle(LOCTEXT("FilterBarAreaTitle", "Log Filter").ToString())
+							.AreaTitle(LOCTEXT("FilterBarAreaTitle", "Log Filter"))
 							.InitiallyCollapsed(true)
 							.Padding(FMargin(8.0f, 6.0f))
 							.BodyContent()
diff --git a/Engine/Source/Editor/DetailCustomizations/Private/AnimTransitionNodeDetails.cpp b/Engine/Source/Editor/DetailCustomizations/Private/AnimTransitionNodeDetails.cpp
--- a/Engine/Source/Editor/DetailCustomizations/Private/AnimTransitionNodeDetails.cpp
+++ b/Engine/Source/Editor/DetailCustomizations/Private/AnimTransitionNodeDetails.cpp
@@ -378,7 +378,7 @@ TSharedRef<SWidget> FAnimTransitionNodeDetails::GetWidgetForInlineShareMenu(FStr
 
 	return
 		SNew(SExpandableArea)
-		.AreaTitle(TypeName)
+		.AreaTitle(FText::FromString(TypeName))
 		.InitiallyCollapsed(true)
 		.BodyContent()
 		[
diff --git a/Engine/Source/Editor/Persona/Private/SAnimCompositePanel.cpp b/Engine/Source/Editor/Persona/Private/SAnimCompositePanel.cpp
--- a/Engine/Source/Editor/Persona/Private/SAnimCompositePanel.cpp
+++ b/Engine/Source/Editor/Persona/Private/SAnimCompositePanel.cpp
@@ -34,7 +34,7 @@ void SAnimCompositePanel::Construct(const FArguments& InArgs)
 		.FillHeight(1)
 		[
 			SNew( SExpandableArea )
-			.AreaTitle( LOCTEXT( "CompositeLabel", "Composite" ).ToString() )
+			.AreaTitle( LOCTEXT( "CompositeLabel", "Composite" ) )
 			.BodyContent()
 			[
 				SAssignNew( PanelArea, SBorder )
