Fixed code relying on SLATE_TEXT_ARGUMENT for STextEntryPopup

Made sure everything was using FText rather than FString.

diff --git a/Engine/Source/Developer/SessionFrontend/Private/Widgets/Console/SSessionConsoleShortcutWindow.cpp b/Engine/Source/Developer/SessionFrontend/Private/Widgets/Console/SSessionConsoleShortcutWindow.cpp
--- a/Engine/Source/Developer/SessionFrontend/Private/Widgets/Console/SSessionConsoleShortcutWindow.cpp
+++ b/Engine/Source/Developer/SessionFrontend/Private/Widgets/Console/SSessionConsoleShortcutWindow.cpp
@@ -63,7 +63,7 @@ void SSessionConsoleShortcutWindow::AddShortcutInternal( const FString& InName,
 }
 
 
-void SSessionConsoleShortcutWindow::HandleEditCommandActionExecute( TSharedPtr<FConsoleShortcutData> InShortcut, bool bInEditCommand, FString InPromptTitle )
+void SSessionConsoleShortcutWindow::HandleEditCommandActionExecute( TSharedPtr<FConsoleShortcutData> InShortcut, bool bInEditCommand, FText InPromptTitle )
 {
 	FString DefaultString = bInEditCommand ? InShortcut->Command : InShortcut->Name;
 
@@ -177,14 +177,14 @@ TSharedRef<ITableRow> SSessionConsoleShortcutWindow::HandleShortcutListViewGener
 			NSLOCTEXT("SessionFrontend", "ContextMenu.EditName", "Edit Name"),
 			FText(),
 			FSlateIcon(),
-			FUIAction(FExecuteAction::CreateSP(this, &SSessionConsoleShortcutWindow::HandleEditCommandActionExecute, InItem, false, LOCTEXT("ShortcutOptionsEditNameTitle", "Name:").ToString()))
+			FUIAction(FExecuteAction::CreateSP(this, &SSessionConsoleShortcutWindow::HandleEditCommandActionExecute, InItem, false, LOCTEXT("ShortcutOptionsEditNameTitle", "Name:")))
 		);
 
 		ContextMenuBuilder.AddMenuEntry(
 			NSLOCTEXT("SessionFrontend", "ContextMenu.EditCommand", "Edit Command"),
 			FText(),
 			FSlateIcon(),
-			FUIAction(FExecuteAction::CreateSP(this, &SSessionConsoleShortcutWindow::HandleEditCommandActionExecute, InItem, true, LOCTEXT("ShortcutOptionsEditCommandTitle", "Command:").ToString()))
+			FUIAction(FExecuteAction::CreateSP(this, &SSessionConsoleShortcutWindow::HandleEditCommandActionExecute, InItem, true, LOCTEXT("ShortcutOptionsEditCommandTitle", "Command:")))
 		);
 	}
 	ContextMenuBuilder.EndSection();
diff --git a/Engine/Source/Developer/SessionFrontend/Private/Widgets/Console/SSessionConsoleShortcutWindow.h b/Engine/Source/Developer/SessionFrontend/Private/Widgets/Console/SSessionConsoleShortcutWindow.h
--- a/Engine/Source/Developer/SessionFrontend/Private/Widgets/Console/SSessionConsoleShortcutWindow.h
+++ b/Engine/Source/Developer/SessionFrontend/Private/Widgets/Console/SSessionConsoleShortcutWindow.h
@@ -88,7 +88,7 @@ private:
 	FReply HandleExecuteButtonClicked( TSharedPtr<FConsoleShortcutData> InShortcut );
 
 	/** Callback for when a shortcut name or command is being edited. */
-	void HandleEditCommandActionExecute( TSharedPtr<FConsoleShortcutData> InShortcut, bool bInEditCommand, FString InPromptTitle );
+	void HandleEditCommandActionExecute( TSharedPtr<FConsoleShortcutData> InShortcut, bool bInEditCommand, FText InPromptTitle );
 
 	/** Generates a row widget for a shortcut. */
 	TSharedRef<ITableRow> HandleShortcutListViewGenerateRow( TSharedPtr<FConsoleShortcutData> InItem, const TSharedRef<STableViewBase>& OwnerTable );
diff --git a/Engine/Source/Editor/Cascade/Private/Cascade.cpp b/Engine/Source/Editor/Cascade/Private/Cascade.cpp
--- a/Engine/Source/Editor/Cascade/Private/Cascade.cpp
+++ b/Engine/Source/Editor/Cascade/Private/Cascade.cpp
@@ -3694,7 +3694,7 @@ void FCascade::OnSetMotionRadius()
 	FString DefaultText = FString::Printf(TEXT("%.2f"), MotionModeRadius);
 	TSharedRef<STextEntryPopup> TextEntry = 
 		SNew(STextEntryPopup)
-		.Label(NSLOCTEXT("Cascade", "MotionRadius", "Motion Radius: ").ToString())
+		.Label(NSLOCTEXT("Cascade", "MotionRadius", "Motion Radius: "))
 		.DefaultText(FText::FromString( DefaultText ) )
 		.OnTextCommitted(this, &FCascade::MotionRadiusCommitted)
 		.SelectAllTextWhenFocused(true)
@@ -3917,7 +3917,7 @@ void FCascade::OnToggleWireframeSphere()
 		FString DefaultText = FString::Printf(TEXT("%.2f"), PreviewViewport->GetViewportClient()->GetWireSphereRadius());
 		TSharedRef<STextEntryPopup> TextEntry = 
 			SNew(STextEntryPopup)
-			.Label(NSLOCTEXT("Cascade", "SphereRadius", "Sphere Radius: ").ToString())
+			.Label(NSLOCTEXT("Cascade", "SphereRadius", "Sphere Radius: "))
 			.DefaultText(FText::FromString(DefaultText))
 			.OnTextCommitted(this, &FCascade::SphereRadiusCommitted)
 			.SelectAllTextWhenFocused(true)
@@ -4824,7 +4824,7 @@ void FCascade::OnRenameEmitter()
 
 	TSharedRef<STextEntryPopup> TextEntry = 
 		SNew(STextEntryPopup)
-		.Label(NSLOCTEXT("Cascade", "SetEmitterName", "Emitter Name: ").ToString())
+		.Label(NSLOCTEXT("Cascade", "SetEmitterName", "Emitter Name: "))
 		.DefaultText(FText::FromName( SelectedEmitter->GetEmitterName() ))
 		.OnTextCommitted(this, &FCascade::EmitterNameCommitted)
 		.SelectAllTextWhenFocused(true)
diff --git a/Engine/Source/Editor/Matinee/Private/MatineeMenus.cpp b/Engine/Source/Editor/Matinee/Private/MatineeMenus.cpp
--- a/Engine/Source/Editor/Matinee/Private/MatineeMenus.cpp
+++ b/Engine/Source/Editor/Matinee/Private/MatineeMenus.cpp
@@ -2878,7 +2878,7 @@ void FMatinee::OnContextSaveAsCameraAnimation()
 
 	TSharedRef<STextEntryPopup> TextEntry = 
 		SNew(STextEntryPopup)
-		.Label(NSLOCTEXT("Matinee.Popups", "ExportCameraAnim_Header", "Export To CameraAnim").ToString())
+		.Label(NSLOCTEXT("Matinee.Popups", "ExportCameraAnim_Header", "Export To CameraAnim"))
 		.DefaultText(  FText::FromString(PackageName) )
 		.OnTextCommitted(this, &FMatinee::ExportCameraAnimationNameCommitted)
 		.ClearKeyboardFocusOnCommit( false );
diff --git a/Engine/Source/Editor/Matinee/Private/MatineeTrackHelpers.cpp b/Engine/Source/Editor/Matinee/Private/MatineeTrackHelpers.cpp
--- a/Engine/Source/Editor/Matinee/Private/MatineeTrackHelpers.cpp
+++ b/Engine/Source/Editor/Matinee/Private/MatineeTrackHelpers.cpp
@@ -401,7 +401,7 @@ bool UMatineeTrackEventHelper::PreCreateKeyframe( UInterpTrack *Track, float Key
 	check(Mode->InterpEd != NULL);
 	TSharedRef<STextEntryPopup> TextEntryPopup = 
 		SNew(STextEntryPopup)
-		.Label(NSLOCTEXT("Matinee.Popups", "NewEventName", "New Event Name").ToString())
+		.Label(NSLOCTEXT("Matinee.Popups", "NewEventName", "New Event Name"))
 		.DefaultText(FText::FromString(TEXT("Event")))
 		.OnTextCommitted_UObject(this, &UMatineeTrackEventHelper::OnAddKeyTextEntry, (IMatineeBase*)Mode->InterpEd, Track)
 		.SelectAllTextWhenFocused(true)
diff --git a/Engine/Source/Editor/Sequencer/Private/Sequencer.cpp b/Engine/Source/Editor/Sequencer/Private/Sequencer.cpp
--- a/Engine/Source/Editor/Sequencer/Private/Sequencer.cpp
+++ b/Engine/Source/Editor/Sequencer/Private/Sequencer.cpp
@@ -288,7 +288,7 @@ void FSequencer::RenameShot(UMovieSceneSection* ShotSection)
 
 	TSharedRef<STextEntryPopup> TextEntry = 
 		SNew(STextEntryPopup)
-		.Label(NSLOCTEXT("Sequencer", "RenameShotHeader", "Name").ToString())
+		.Label(NSLOCTEXT("Sequencer", "RenameShotHeader", "Name"))
 		.DefaultText( ActualShotSection->GetTitle() )
 		.OnTextCommitted(this, &FSequencer::RenameShotCommitted, ShotSection)
 		.ClearKeyboardFocusOnCommit( false );
