SessionConsole: Code and documentation cleanup pass

diff --git a/Engine/Source/Developer/SessionFrontend/Private/Widgets/Console/SSessionConsole.cpp b/Engine/Source/Developer/SessionFrontend/Private/Widgets/Console/SSessionConsole.cpp
--- a/Engine/Source/Developer/SessionFrontend/Private/Widgets/Console/SSessionConsole.cpp
+++ b/Engine/Source/Developer/SessionFrontend/Private/Widgets/Console/SSessionConsole.cpp
@@ -24,7 +24,8 @@ SSessionConsole::~SSessionConsole()
 /* SSessionConsolePanel interface
  *****************************************************************************/
 
-void SSessionConsole::Construct( const FArguments& InArgs, ISessionManagerRef InSessionManager )
+BEGIN_SLATE_FUNCTION_BUILD_OPTIMIZATION
+void SSessionConsole::Construct(const FArguments& InArgs, ISessionManagerRef InSessionManager)
 {
 	SessionManager = InSessionManager;
 	ShouldScrollToLast = true;
@@ -111,11 +112,11 @@ void SSessionConsole::Construct( const FArguments& InArgs, ISessionManagerRef In
 
 						//Shortcut buttons
 						+ SHorizontalBox::Slot()
-						.FillWidth(0.2f)
-						[
-							SAssignNew(ShortcutWindow, SSessionConsoleShortcutWindow)
-								.OnCommandSubmitted(this, &SSessionConsole::HandleCommandSubmitted)
-						]
+							.FillWidth(0.2f)
+							[
+								SAssignNew(ShortcutWindow, SSessionConsoleShortcutWindow)
+									.OnCommandSubmitted(this, &SSessionConsole::HandleCommandSubmitted)
+							]
 					]
 
 				+ SVerticalBox::Slot()
@@ -155,6 +156,7 @@ void SSessionConsole::Construct( const FArguments& InArgs, ISessionManagerRef In
 
 	ReloadLog(true);
 }
+END_SLATE_FUNCTION_BUILD_OPTIMIZATION
 
 
 /* SSessionConsolePanel implementation
@@ -194,26 +196,24 @@ void SSessionConsole::CopyLog()
 {
 	TArray<FSessionLogMessagePtr> SelectedItems = LogListView->GetSelectedItems();
 
-	if (SelectedItems.Num() > 0)
+	if (SelectedItems.Num() == 0)
 	{
-		FString SelectedText;
+		return;
+	}
 
-		for( int32 Index = 0; Index < SelectedItems.Num(); ++Index )
-		{
-			SelectedText += FString::Printf(TEXT("%s [%s] %09.3f: %s"),
-				*SelectedItems[Index]->Time.ToString(),
-				*SelectedItems[Index]->InstanceName,
-				SelectedItems[Index]->TimeSeconds,
-				*SelectedItems[Index]->Text
-			) + LINE_TERMINATOR;
-		}
+	FString SelectedText;
 
-		FPlatformMisc::ClipboardCopy( *SelectedText );
+	for (const auto& Item : SelectedItems)
+	{
+		SelectedText += FString::Printf(TEXT("%s [%s] %09.3f: %s"), *Item->Time.ToString(), *Item->InstanceName, Item->TimeSeconds, *Item->Text);
+		SelectedText += LINE_TERMINATOR;
 	}
+
+	FPlatformMisc::ClipboardCopy(*SelectedText);
 }
 
 
-void SSessionConsole::ReloadLog( bool FullyReload )
+void SSessionConsole::ReloadLog(bool FullyReload)
 {
 	// reload log list
 	if (FullyReload)
@@ -223,13 +223,13 @@ void SSessionConsole::ReloadLog( bool FullyReload )
 		TArray<ISessionInstanceInfoPtr> SelectedInstances;
 		SessionManager->GetSelectedInstances(SelectedInstances);
 
-		for (int32 InstanceIndex = 0; InstanceIndex < SelectedInstances.Num(); ++InstanceIndex)
+		for (const auto& Instance : SelectedInstances)
 		{
-			const TArray<FSessionLogMessagePtr>& InstanceLog = SelectedInstances[InstanceIndex]->GetLog();
+			const TArray<FSessionLogMessagePtr>& InstanceLog = Instance->GetLog();
 
-			for (int32 LogIndex = 0; LogIndex < InstanceLog.Num(); ++LogIndex)
+			for (const auto& LogMessage : InstanceLog)
 			{
-				AvailableLogs.HeapPush(InstanceLog[LogIndex], FSessionLogMessage::TimeComparer());
+				AvailableLogs.HeapPush(LogMessage, FSessionLogMessage::TimeComparer());
 			}
 		}
 
@@ -241,10 +241,8 @@ void SSessionConsole::ReloadLog( bool FullyReload )
 	// filter log list
 	FilterBar->ResetFilter();
 
-	for (int32 Index = 0; Index < AvailableLogs.Num(); ++Index)
+	for (const auto& LogMessage : AvailableLogs)
 	{
-		const FSessionLogMessagePtr& LogMessage = AvailableLogs[Index];
-
 		if (FilterBar->FilterLogMessage(LogMessage.ToSharedRef()))
 		{
 			LogMessages.Add(LogMessage);
@@ -265,81 +263,87 @@ void SSessionConsole::SaveLog()
 {
 	IDesktopPlatform* DesktopPlatform = FDesktopPlatformModule::Get();
 
-	if (DesktopPlatform != NULL)
+	if (DesktopPlatform == nullptr)
+	{
+		FMessageDialog::Open(EAppMsgType::Ok, LOCTEXT("SaveLogDialogUnsupportedError", "Saving is not supported on this platform!"));
+
+		return;
+	}
+
+	TArray<FString> Filenames;
+
+	// open file dialog
+	TSharedPtr<SWindow> ParentWindow = FSlateApplication::Get().FindWidgetWindow(AsShared());
+	void* ParentWindowHandle = (ParentWindow.IsValid() && ParentWindow->GetNativeWindow().IsValid()) ? ParentWindow->GetNativeWindow()->GetOSWindowHandle() : nullptr;
+
+	if (!DesktopPlatform->SaveFileDialog(
+		ParentWindowHandle,
+		LOCTEXT("SaveLogDialogTitle", "Save Log As...").ToString(),
+		LastLogFileSaveDirectory,
+		TEXT("Session.log"),
+		TEXT("Log Files (*.log)|*.log"),
+		EFileDialogFlags::None,
+		Filenames))
+	{
+		return;
+	}
+
+	// no log file selected?
+	if (Filenames.Num() == 0)
+	{
+		return;
+	}
+
+	FString Filename = Filenames[0];
+
+	// keep path as default for next time
+	LastLogFileSaveDirectory = FPaths::GetPath(Filename);
+
+	// add a file extension if none was provided
+	if (FPaths::GetExtension(Filename).IsEmpty())
+	{
+		Filename += Filename + TEXT(".log");
+	}
+
+	// save file
+	FArchive* LogFile = IFileManager::Get().CreateFileWriter(*Filename);
+
+	if (LogFile != nullptr)
 	{
-		TArray<FString> Filenames;
-
-		TSharedPtr<SWindow> ParentWindow = FSlateApplication::Get().FindWidgetWindow(AsShared());
-		void* ParentWindowHandle = (ParentWindow.IsValid() && ParentWindow->GetNativeWindow().IsValid()) ? ParentWindow->GetNativeWindow()->GetOSWindowHandle() : nullptr;
-
-		if (DesktopPlatform->SaveFileDialog(
-			ParentWindowHandle,
-			LOCTEXT("SaveLogDialogTitle", "Save Log As...").ToString(),
-			LastLogFileSaveDirectory,
-			TEXT("Session.log"),
-			TEXT("Log Files (*.log)|*.log"),
-			EFileDialogFlags::None,
-			Filenames))
+		for (const auto& LogMessage : LogMessages)
 		{
-			if (Filenames.Num() > 0)
-			{
-				FString Filename = Filenames[0];
-
-				// keep path as default for next time
-				LastLogFileSaveDirectory = FPaths::GetPath(Filename);
-
-				// add a file extension if none was provided
-				if (FPaths::GetExtension(Filename).IsEmpty())
-				{
-					Filename += Filename + TEXT(".log");
-				}
-
-				// save file
-				FArchive* LogFile = IFileManager::Get().CreateFileWriter(*Filename);
-
-				if (LogFile != NULL)
-				{
-					for( int32 Index = 0; Index < LogMessages.Num(); ++Index )
-					{
-						FString LogEntry = FString::Printf(TEXT("%s [%s] %09.3f: %s"),
-							*LogMessages[Index]->Time.ToString(),
-							*LogMessages[Index]->InstanceName,
-							LogMessages[Index]->TimeSeconds,
-							*LogMessages[Index]->Text) + LINE_TERMINATOR;
-
-						LogFile->Serialize(TCHAR_TO_ANSI(*LogEntry), LogEntry.Len());
-					}
-
-					LogFile->Close();
-
-					delete LogFile;
-				}
-				else
-				{
-					FMessageDialog::Open(EAppMsgType::Ok, LOCTEXT("SaveLogDialogFileError", "Failed to open the specified file for saving!"));
-				}
-			}
+			FString LogEntry = FString::Printf(TEXT("%s [%s] %09.3f: %s"),
+				*LogMessage->Time.ToString(),
+				*LogMessage->InstanceName,
+				LogMessage->TimeSeconds,
+				*LogMessage->Text) + LINE_TERMINATOR;
+
+			LogFile->Serialize(TCHAR_TO_ANSI(*LogEntry), LogEntry.Len());
 		}
+
+		LogFile->Close();
+		delete LogFile;
 	}
 	else
 	{
-		FMessageDialog::Open(EAppMsgType::Ok, LOCTEXT("SaveLogDialogUnsupportedError", "Saving is not supported on this platform!"));
+		FMessageDialog::Open(EAppMsgType::Ok, LOCTEXT("SaveLogDialogFileError", "Failed to open the specified file for saving!"));
 	}
 }
 
 
-void SSessionConsole::SendCommand( const FString& CommandString )
+void SSessionConsole::SendCommand(const FString& CommandString)
 {
-	if (!CommandString.IsEmpty())
+	if (CommandString.IsEmpty())
 	{
-		TArray<ISessionInstanceInfoPtr> SelectedInstances;
+		return;
+	}
 
-		SessionManager->GetSelectedInstances(SelectedInstances);
+	TArray<ISessionInstanceInfoPtr> SelectedInstances;
+	SessionManager->GetSelectedInstances(SelectedInstances);
 
-		for (int32 Index = 0; Index < SelectedInstances.Num(); Index++)
-		{
-			SelectedInstances[Index]->ExecuteCommand(CommandString);
-		}
+	for (auto& Instance : SelectedInstances)
+	{
+		Instance->ExecuteCommand(CommandString);
 	}
 }
 
@@ -348,7 +352,7 @@ void SSessionConsole::SendCommand( const FString& CommandString )
 /* SWidget implementation
  *****************************************************************************/
 
-FReply SSessionConsole::OnKeyDown( const FGeometry& MyGeometry, const FKeyEvent& InKeyEvent )
+FReply SSessionConsole::OnKeyDown(const FGeometry& MyGeometry, const FKeyEvent& InKeyEvent)
 {
 	if (InKeyEvent.IsControlDown())
 	{
@@ -358,7 +362,8 @@ FReply SSessionConsole::OnKeyDown( const FGeometry& MyGeometry, const FKeyEvent&
 
 			return FReply::Handled();
 		}
-		else if (InKeyEvent.GetKey() == EKeys::S)
+	
+		if (InKeyEvent.GetKey() == EKeys::S)
 		{
 			SaveLog();
 
@@ -385,13 +390,13 @@ bool SSessionConsole::HandleClearActionCanExecute()
 }
 
 
-void SSessionConsole::HandleCommandBarPromoteToShortcutClicked( const FString& CommandString )
+void SSessionConsole::HandleCommandBarPromoteToShortcutClicked(const FString& CommandString)
 {
 	ShortcutWindow->AddShortcut(CommandString, CommandString);
 }
 
 
-void SSessionConsole::HandleCommandSubmitted( const FString& CommandString )
+void SSessionConsole::HandleCommandSubmitted(const FString& CommandString)
 {
 	SendCommand(CommandString);
 }
@@ -417,7 +422,7 @@ void SSessionConsole::HandleFilterChanged()
 }
 
 
-void SSessionConsole::HandleLogListItemScrolledIntoView( FSessionLogMessagePtr Item, const TSharedPtr<ITableRow>& TableRow )
+void SSessionConsole::HandleLogListItemScrolledIntoView(FSessionLogMessagePtr Item, const TSharedPtr<ITableRow>& TableRow)
 {
 	if (LogMessages.Num() > 0)
 	{
@@ -430,7 +435,7 @@ void SSessionConsole::HandleLogListItemScrolledIntoView( FSessionLogMessagePtr I
 }
 
 
-TSharedRef<ITableRow> SSessionConsole::HandleLogListGenerateRow( FSessionLogMessagePtr Message, const TSharedRef<STableViewBase>& OwnerTable )
+TSharedRef<ITableRow> SSessionConsole::HandleLogListGenerateRow(FSessionLogMessagePtr Message, const TSharedRef<STableViewBase>& OwnerTable)
 {
 	return SNew(SSessionConsoleLogTableRow, OwnerTable)
 		.HighlightText(this, &SSessionConsole::HandleLogListGetHighlightText)
@@ -480,27 +485,26 @@ void SSessionConsole::HandleSessionManagerInstanceSelectionChanged()
 }
 
 
-void SSessionConsole::HandleSessionManagerLogReceived( const ISessionInfoRef& Session, const ISessionInstanceInfoRef& Instance, const FSessionLogMessageRef& Message)
+void SSessionConsole::HandleSessionManagerLogReceived(const ISessionInfoRef& Session, const ISessionInstanceInfoRef& Instance, const FSessionLogMessageRef& Message)
 {
-	if (SessionManager->IsInstanceSelected(Instance))
+	if (!SessionManager->IsInstanceSelected(Instance) || !FilterBar->FilterLogMessage(Message))
 	{
-		if (FilterBar->FilterLogMessage(Message))
-		{
-			AvailableLogs.Add(Message);
-			LogMessages.Add(Message);
+		return;
+	}
 
-			LogListView->RequestListRefresh();
+	AvailableLogs.Add(Message);
+	LogMessages.Add(Message);
 
-			if (ShouldScrollToLast)
-			{
-				LogListView->RequestScrollIntoView(Message);
-			}
-		}
+	LogListView->RequestListRefresh();
+
+	if (ShouldScrollToLast)
+	{
+		LogListView->RequestScrollIntoView(Message);
 	}
 }
 
 
-void SSessionConsole::HandleSessionManagerSelectedSessionChanged( const ISessionInfoPtr& SelectedSession )
+void SSessionConsole::HandleSessionManagerSelectedSessionChanged(const ISessionInfoPtr& SelectedSession)
 {
 	ReloadLog(true);
 }
diff --git a/Engine/Source/Developer/SessionFrontend/Private/Widgets/Console/SSessionConsole.h b/Engine/Source/Developer/SessionFrontend/Private/Widgets/Console/SSessionConsole.h
--- a/Engine/Source/Developer/SessionFrontend/Private/Widgets/Console/SSessionConsole.h
+++ b/Engine/Source/Developer/SessionFrontend/Private/Widgets/Console/SSessionConsole.h
@@ -30,27 +30,40 @@ public:
 	 * @param InArgs The declaration data for this widget.
 	 * @param InSessionManager The session manager to use.
 	 */
-	void Construct( const FArguments& InArgs, ISessionManagerRef InSessionManager );
+	void Construct(const FArguments& InArgs, ISessionManagerRef InSessionManager);
 
 protected:
 
 	/** Binds the device commands on our toolbar. */
 	void BindCommands();
 
-	/** Clears the log list view. */
+	/**
+	 * Clears the log list view.
+	 *
+	 * @see CopyLog, ReloadLog, SaveLog
+	 */
 	void ClearLog();
 
-	/** Copies the selected log messages to the clipboard. */
+	/**
+	 * Copies the selected log messages to the clipboard.
+	 *
+	 * @see ClearLog, ReloadLog, SaveLog
+	 */
 	void CopyLog();
 
 	/**
 	 * Reloads the log messages for the currently selected engine instances.
 	 *
-	 * @param FullyReload - Whether to fully reload the log entries or only re-apply filtering.
+	 * @param FullyReload Whether to fully reload the log entries or only re-apply filtering.
+	 * @see ClearLog, CopyLog, SaveLog
 	 */
-	void ReloadLog( bool FullyReload );
+	void ReloadLog(bool FullyReload);
 
-	/** Saves all log messages to a file. */
+	/**
+	 * Saves all log messages to a file.
+	 *
+	 * @see ClearLog, CopyLog, ReloadLog
+	 */
 	void SaveLog();
 
 	/**
@@ -58,11 +71,13 @@ protected:
 	 *
 	 * @param CommandString The command string to send.
 	 */
-	void SendCommand( const FString& CommandString );
+	void SendCommand(const FString& CommandString);
 
 protected:
 
-	virtual FReply OnKeyDown( const FGeometry& MyGeometry, const FKeyEvent& InKeyEvent ) override;
+	// SCompoundWidget overrides
+
+	virtual FReply OnKeyDown(const FGeometry& MyGeometry, const FKeyEvent& InKeyEvent) override;
 
 private:
 
@@ -85,25 +100,25 @@ private:
 	bool HandleSaveActionCanExecute();
 
 	/** Callback for promoting console command to shortcuts. */
-	void HandleCommandBarPromoteToShortcutClicked( const FString& CommandString );
+	void HandleCommandBarPromoteToShortcutClicked(const FString& CommandString);
 
 	/** Callback for submitting console commands. */
-	void HandleCommandSubmitted( const FString& CommandString );
+	void HandleCommandSubmitted(const FString& CommandString);
 
 	/** Callback for changing the filter settings. */
 	void HandleFilterChanged();
 
 	/** Callback for scrolling a log item into view. */
-	void HandleLogListItemScrolledIntoView( FSessionLogMessagePtr Item, const TSharedPtr<ITableRow>& TableRow );
+	void HandleLogListItemScrolledIntoView(FSessionLogMessagePtr Item, const TSharedPtr<ITableRow>& TableRow);
 
 	/** Callback for generating a row widget for the log list view. */
-	TSharedRef<ITableRow> HandleLogListGenerateRow( FSessionLogMessagePtr Message, const TSharedRef<STableViewBase>& OwnerTable );
+	TSharedRef<ITableRow> HandleLogListGenerateRow(FSessionLogMessagePtr Message, const TSharedRef<STableViewBase>& OwnerTable);
 
 	/** Callback for getting the highlight string for log messages. */
 	FText HandleLogListGetHighlightText() const;
 
 	/** Callback for selecting log messages. */
-	void HandleLogListSelectionChanged( FSessionLogMessagePtr InItem, ESelectInfo::Type SelectInfo );
+	void HandleLogListSelectionChanged(FSessionLogMessagePtr InItem, ESelectInfo::Type SelectInfo);
 
 	/** Callback for getting the enabled state of the console box. */
 	bool HandleMainContentIsEnabled() const;
@@ -115,10 +130,10 @@ private:
 	void HandleSessionManagerInstanceSelectionChanged();
 
 	/** Callback for received log entries. */
-	void HandleSessionManagerLogReceived( const ISessionInfoRef& Session, const ISessionInstanceInfoRef& Instance, const FSessionLogMessageRef& Message );
+	void HandleSessionManagerLogReceived(const ISessionInfoRef& Session, const ISessionInstanceInfoRef& Instance, const FSessionLogMessageRef& Message);
 
 	/** Callback for changing the selected session. */
-	void HandleSessionManagerSelectedSessionChanged( const ISessionInfoPtr& SelectedSession );
+	void HandleSessionManagerSelectedSessionChanged(const ISessionInfoPtr& SelectedSession);
 
 private:
 
diff --git a/Engine/Source/Developer/SessionFrontend/Private/Widgets/Console/SSessionConsoleCommandBar.cpp b/Engine/Source/Developer/SessionFrontend/Private/Widgets/Console/SSessionConsoleCommandBar.cpp
--- a/Engine/Source/Developer/SessionFrontend/Private/Widgets/Console/SSessionConsoleCommandBar.cpp
+++ b/Engine/Source/Developer/SessionFrontend/Private/Widgets/Console/SSessionConsoleCommandBar.cpp
@@ -10,7 +10,8 @@
 /* SSessionConsoleCommandBar interface
  *****************************************************************************/
 
-void SSessionConsoleCommandBar::Construct( const FArguments& InArgs )
+BEGIN_SLATE_FUNCTION_BUILD_OPTIMIZATION
+void SSessionConsoleCommandBar::Construct(const FArguments& InArgs)
 {
 	OnCommandSubmitted = InArgs._OnCommandSubmitted;
 	OnPromoteToShortcutClicked = InArgs._OnPromoteToShortcutClicked;
@@ -64,9 +65,10 @@ void SSessionConsoleCommandBar::Construct( const FArguments& InArgs )
 			]
 	];
 }
+END_SLATE_FUNCTION_BUILD_OPTIMIZATION
 
 
-void SSessionConsoleCommandBar::SetNumSelectedInstances( int Count )
+void SSessionConsoleCommandBar::SetNumSelectedInstances(int Count)
 {
 	FString CommandString = InputTextBox->GetText().ToString();
 	CommandString.Trim();
@@ -82,7 +84,7 @@ void SSessionConsoleCommandBar::SetNumSelectedInstances( int Count )
 /* SSessionConsoleCommandBar implementation
  *****************************************************************************/
 
-void SSessionConsoleCommandBar::SubmitCommand( const FString& Command )
+void SSessionConsoleCommandBar::SubmitCommand(const FString& Command)
 {
 	OnCommandSubmitted.ExecuteIfBound(Command);
 
@@ -96,7 +98,7 @@ void SSessionConsoleCommandBar::SubmitCommand( const FString& Command )
 /* SSessionConsoleCommandBar event handlers
  *****************************************************************************/
 
-void SSessionConsoleCommandBar::HandleInputTextChanged( const FText& InText )
+void SSessionConsoleCommandBar::HandleInputTextChanged(const FText& InText)
 {
 	FString CommandString = InputTextBox->GetText().ToString();
 	CommandString.Trim();
@@ -106,7 +108,7 @@ void SSessionConsoleCommandBar::HandleInputTextChanged( const FText& InText )
 }
 
 
-void SSessionConsoleCommandBar::HandleInputTextCommitted( const FText& InText, ETextCommit::Type CommitInfo )
+void SSessionConsoleCommandBar::HandleInputTextCommitted(const FText& InText, ETextCommit::Type CommitInfo)
 {
 	if (CommitInfo == ETextCommit::OnEnter)
 	{
@@ -115,13 +117,13 @@ void SSessionConsoleCommandBar::HandleInputTextCommitted( const FText& InText, E
 }
 
 
-void SSessionConsoleCommandBar::HandleInputTextShowingHistory( TArray<FString>& OutHistory )
+void SSessionConsoleCommandBar::HandleInputTextShowingHistory(TArray<FString>& OutHistory)
 {
 	OutHistory = CommandHistory;
 }
 
 
-void SSessionConsoleCommandBar::HandleInputTextShowingSuggestions( const FString& Text, TArray<FString>& OutSuggestions )
+void SSessionConsoleCommandBar::HandleInputTextShowingSuggestions(const FString& Text, TArray<FString>& OutSuggestions)
 {
 	// @todo gmp: implement remote auto-complete
 }
diff --git a/Engine/Source/Developer/SessionFrontend/Private/Widgets/Console/SSessionConsoleCommandBar.h b/Engine/Source/Developer/SessionFrontend/Private/Widgets/Console/SSessionConsoleCommandBar.h
--- a/Engine/Source/Developer/SessionFrontend/Private/Widgets/Console/SSessionConsoleCommandBar.h
+++ b/Engine/Source/Developer/SessionFrontend/Private/Widgets/Console/SSessionConsoleCommandBar.h
@@ -3,12 +3,8 @@
 #pragma once
 
 
-/**
- * Delegate type for submitting console commands.
- *
- * The first parameter is the submitted command string.
- */
-DECLARE_DELEGATE_OneParam(FOnSessionConsoleCommandSubmitted, const FString&)
+/** Delegate type for submitting console commands. */
+DECLARE_DELEGATE_OneParam(FOnSessionConsoleCommandSubmitted, const FString& /*Command*/)
 
 
 /**
@@ -36,14 +32,14 @@ public:
 	 *
 	 * @param InArgs The declaration data for this widget.
 	 */
-	void Construct( const FArguments& InArgs );
+	void Construct(const FArguments& InArgs);
 
 	/**
 	 * Sets the number of selected engine instances.
 	 *
 	 * @param Count Number of selected instances.
 	 */
-	void SetNumSelectedInstances( int Count );
+	void SetNumSelectedInstances(int Count);
 
 protected:
 
@@ -52,21 +48,21 @@ protected:
 	 *
 	 * @param Command the command to submit.
 	 */
-	void SubmitCommand( const FString& Command );
+	void SubmitCommand(const FString& Command);
 
 private:
 
 	/** Handles changing the input text box's content. */
-	void HandleInputTextChanged( const FText& InText );
+	void HandleInputTextChanged(const FText& InText);
 
 	/** Handles committing the input text box's content. */
-	void HandleInputTextCommitted( const FText& InText, ETextCommit::Type CommitInfo );
+	void HandleInputTextCommitted(const FText& InText, ETextCommit::Type CommitInfo);
 
 	/** Handles showing a history in the input text box. */
-	void HandleInputTextShowingHistory( TArray<FString>& OutHistory );
+	void HandleInputTextShowingHistory(TArray<FString>& OutHistory);
 
 	/** Handles showing suggestions in the input text box. */
-	void HandleInputTextShowingSuggestions( const FString& Text, TArray<FString>& OutSuggestions );
+	void HandleInputTextShowingSuggestions(const FString& Text, TArray<FString>& OutSuggestions);
 
 	/** Handles clicking the promote to shortcut button. */
 	FReply HandlePromoteToShortcutButtonClicked();
diff --git a/Engine/Source/Developer/SessionFrontend/Private/Widgets/Console/SSessionConsoleLogTableRow.h b/Engine/Source/Developer/SessionFrontend/Private/Widgets/Console/SSessionConsoleLogTableRow.h
--- a/Engine/Source/Developer/SessionFrontend/Private/Widgets/Console/SSessionConsoleLogTableRow.h
+++ b/Engine/Source/Developer/SessionFrontend/Private/Widgets/Console/SSessionConsoleLogTableRow.h
@@ -23,7 +23,7 @@ public:
 	 *
 	 * @param InArgs The construction arguments.
 	 */
-	void Construct( const FArguments& InArgs, const TSharedRef<STableViewBase>& InOwnerTableView )
+	void Construct(const FArguments& InArgs, const TSharedRef<STableViewBase>& InOwnerTableView)
 	{
 		HighlightText = InArgs._HighlightText;
 		LogMessage = InArgs._LogMessage;
@@ -35,7 +35,8 @@ public:
 
 	// SMultiColumnTableRow interface
 
-	virtual TSharedRef<SWidget> GenerateWidgetForColumn( const FName& ColumnName ) override
+	BEGIN_SLATE_FUNCTION_BUILD_OPTIMIZATION
+	virtual TSharedRef<SWidget> GenerateWidgetForColumn(const FName& ColumnName) override
 	{
 		if (ColumnName == "Instance")
 		{
@@ -113,6 +114,7 @@ public:
 
 		return SNullWidget::NullWidget;
 	}
+	END_SLATE_FUNCTION_BUILD_OPTIMIZATION
 
 private:
 
diff --git a/Engine/Source/Developer/SessionFrontend/Private/Widgets/Console/SSessionConsoleShortcutWindow.cpp b/Engine/Source/Developer/SessionFrontend/Private/Widgets/Console/SSessionConsoleShortcutWindow.cpp
--- a/Engine/Source/Developer/SessionFrontend/Private/Widgets/Console/SSessionConsoleShortcutWindow.cpp
+++ b/Engine/Source/Developer/SessionFrontend/Private/Widgets/Console/SSessionConsoleShortcutWindow.cpp
@@ -10,14 +10,15 @@
 /* SSessionConsoleShortcutWindow interface
  *****************************************************************************/
 
-void SSessionConsoleShortcutWindow::AddShortcut( const FString& InName, const FString& InCommandString )
+void SSessionConsoleShortcutWindow::AddShortcut(const FString& InName, const FString& InCommandString)
 {
 	AddShortcutInternal(InName, InCommandString);
 	SaveShortcuts();
 }
 
 
-void SSessionConsoleShortcutWindow::Construct( const FArguments& InArgs )
+BEGIN_SLATE_FUNCTION_BUILD_OPTIMIZATION
+void SSessionConsoleShortcutWindow::Construct(const FArguments& InArgs)
 {
 	OnCommandSubmitted = InArgs._OnCommandSubmitted;
 
@@ -47,23 +48,24 @@ void SSessionConsoleShortcutWindow::Construct( const FArguments& InArgs )
 	LoadShortcuts();
 	RebuildUI();
 }
+END_SLATE_FUNCTION_BUILD_OPTIMIZATION
 
 
 /* SSessionConsoleShortcutWindow implementation
  *****************************************************************************/
 
-void SSessionConsoleShortcutWindow::AddShortcutInternal( const FString& InName, const FString& InCommandString )
+void SSessionConsoleShortcutWindow::AddShortcutInternal(const FString& InName, const FString& InCommandString)
 {
 	TSharedPtr<FConsoleShortcutData> NewCommand = MakeShareable(new FConsoleShortcutData());
 	NewCommand->Name = InName;
 	NewCommand->Command = InCommandString;
 
-	Shortcuts.Add( NewCommand );
+	Shortcuts.Add(NewCommand);
 	RebuildUI();
 }
 
 
-void SSessionConsoleShortcutWindow::HandleEditCommandActionExecute( TSharedPtr<FConsoleShortcutData> InShortcut, bool bInEditCommand, FText InPromptTitle )
+void SSessionConsoleShortcutWindow::HandleEditCommandActionExecute(TSharedPtr<FConsoleShortcutData> InShortcut, bool bInEditCommand, FText InPromptTitle)
 {
 	FString DefaultString = bInEditCommand ? InShortcut->Command : InShortcut->Name;
 
@@ -104,8 +106,8 @@ void SSessionConsoleShortcutWindow::LoadShortcuts()
 	FFileHelper::LoadFileToString(Content, *GetShortcutFilename());
 
 	TSharedPtr<FJsonObject> ShortcutStream;
-	TSharedRef<TJsonReader<>> Reader = TJsonReaderFactory<>::Create( Content );
-	bool bResult = FJsonSerializer::Deserialize( Reader, ShortcutStream );
+	TSharedRef<TJsonReader<>> Reader = TJsonReaderFactory<>::Create(Content);
+	bool bResult = FJsonSerializer::Deserialize(Reader, ShortcutStream);
 
 	if (ShortcutStream.IsValid())
 	{
@@ -128,7 +130,7 @@ void SSessionConsoleShortcutWindow::RebuildUI()
 }
 
 
-void SSessionConsoleShortcutWindow::HandleDeleteCommandActionExecute( TSharedPtr<FConsoleShortcutData> InShortcut )
+void SSessionConsoleShortcutWindow::HandleDeleteCommandActionExecute(TSharedPtr<FConsoleShortcutData> InShortcut)
 {
 	Shortcuts.Remove(InShortcut);
 	RebuildUI();
@@ -137,7 +139,7 @@ void SSessionConsoleShortcutWindow::HandleDeleteCommandActionExecute( TSharedPtr
 
 void SSessionConsoleShortcutWindow::SaveShortcuts() const
 {
-	TSharedPtr<FJsonObject> ShortcutStream = MakeShareable( new FJsonObject );
+	TSharedPtr<FJsonObject> ShortcutStream = MakeShareable(new FJsonObject);
 
 	ShortcutStream->SetNumberField(TEXT("Count"), Shortcuts.Num());
 	for (int32 i = 0; i < Shortcuts.Num(); ++i)
@@ -157,7 +159,7 @@ void SSessionConsoleShortcutWindow::SaveShortcuts() const
 /* SSessionConsoleShortcutWindow callbacks
  *****************************************************************************/
 
-FReply SSessionConsoleShortcutWindow::HandleExecuteButtonClicked( TSharedPtr<FConsoleShortcutData> InShortcut )
+FReply SSessionConsoleShortcutWindow::HandleExecuteButtonClicked(TSharedPtr<FConsoleShortcutData> InShortcut)
 {
 	if (OnCommandSubmitted.IsBound())
 	{
@@ -240,16 +242,17 @@ TSharedRef<ITableRow> SSessionConsoleShortcutWindow::HandleShortcutListViewGener
 }
 
 
-void SSessionConsoleShortcutWindow::HandleShortcutTextEntryCommitted( const FText& CommandText, ETextCommit::Type CommitInfo )
+void SSessionConsoleShortcutWindow::HandleShortcutTextEntryCommitted(const FText& CommandText, ETextCommit::Type CommitInfo)
 {
 	if (NameEntryMenu.IsValid())
 	{
 		NameEntryMenu.Pin()->Dismiss();
 
 		int32 IndexOfShortcut = Shortcuts.IndexOfByKey(EditedShortcut);
+
 		if (EditedShortcut.IsValid() && (IndexOfShortcut != INDEX_NONE))
 		{
-			//make a new version of the command so the list view knows to refresh
+			// make a new version of the command so the list view knows to refresh
 			TSharedPtr<FConsoleShortcutData> NewCommand = MakeShareable(new FConsoleShortcutData());
 			*NewCommand = *EditedShortcut;
 			Shortcuts[IndexOfShortcut] = NewCommand;
