SessionFrontend: Redesigned session browser

diff --git a/Engine/Source/Developer/Profiler/Private/ProfilerManager.cpp b/Engine/Source/Developer/Profiler/Private/ProfilerManager.cpp
--- a/Engine/Source/Developer/Profiler/Private/ProfilerManager.cpp
+++ b/Engine/Source/Developer/Profiler/Private/ProfilerManager.cpp
@@ -398,10 +398,10 @@ void FProfilerManager::SessionManager_OnCanSelectSession( const ISessionInfoPtr&
 
 void FProfilerManager::SessionManager_OnSelectedSessionChanged( const ISessionInfoPtr& InActiveSession )
 {
-	SessionManager_OnInstanceSelectionChanged();
+	SessionManager_OnInstanceSelectionChanged(nullptr, false);
 }
 
-void FProfilerManager::SessionManager_OnInstanceSelectionChanged()
+void FProfilerManager::SessionManager_OnInstanceSelectionChanged(const TSharedPtr<ISessionInstanceInfo>& /*Instance*/, bool /*Selected*/)
 {
 	const ISessionInfoPtr& SelectedSession = SessionManager->GetSelectedSession();
 	const bool SessionIsValid = SelectedSession.IsValid() && (SelectedSession->GetSessionOwner() == FPlatformProcess::UserName(true));
@@ -427,7 +427,7 @@ void FProfilerManager::SessionManager_OnInstanceSelectionChanged()
 	if( ActiveSession.IsValid() )
 	{
 		// Track all selected session instances.
-		SessionManager->GetSelectedInstances( SelectedSessionInstances );
+		SelectedSessionInstances = SessionManager->GetSelectedInstances();
 		const int32 NumSelectedInstances = SelectedSessionInstances.Num();
 		const int32 NumInstances = FMath::Min( NumSelectedInstances, FProfilerManager::GetSettings().bSingleInstanceMode ? 1 : NumSelectedInstances );
 
diff --git a/Engine/Source/Developer/Profiler/Private/ProfilerManager.h b/Engine/Source/Developer/Profiler/Private/ProfilerManager.h
--- a/Engine/Source/Developer/Profiler/Private/ProfilerManager.h
+++ b/Engine/Source/Developer/Profiler/Private/ProfilerManager.h
@@ -530,7 +530,7 @@ protected:
 	void ProfilerClient_OnProfilerFileTransfer( const FString& Filename, int64 FileProgress, int64 FileSize );
 
 	void SessionManager_OnCanSelectSession( const ISessionInfoPtr& Session, bool& CanSelect );
-	void SessionManager_OnInstanceSelectionChanged();
+	void SessionManager_OnInstanceSelectionChanged( const TSharedPtr<ISessionInstanceInfo>& Instance, bool Selected );
 	void SessionManager_OnSelectedSessionChanged( const ISessionInfoPtr& Session );
 
 public:
diff --git a/Engine/Source/Developer/SessionFrontend/Private/Models/SessionBrowserTreeItems.h b/Engine/Source/Developer/SessionFrontend/Private/Models/SessionBrowserTreeItems.h
new file mode 100644
--- /dev/null
+++ b/Engine/Source/Developer/SessionFrontend/Private/Models/SessionBrowserTreeItems.h
@@ -0,0 +1,223 @@
+// Copyright 1998-2015 Epic Games, Inc. All Rights Reserved.
+
+#pragma once
+
+
+enum class ESessionBrowserTreeNodeType
+{
+	Group,
+	Instance,
+	Session
+};
+
+
+/**
+ * Base class for items in the session tree view.
+ */
+class FSessionBrowserTreeItem
+{
+public:
+
+	/**
+	 * Adds a child process item to this item.
+	 *
+	 * @param The child item to add.
+	 * @see ClearChildren, GetChildren
+	 */
+	void AddChild(const TSharedRef<FSessionBrowserTreeItem>& Child)
+	{
+		Children.Add(Child);
+	}
+
+	/**
+	 * Clears the collection of child items.
+	 *
+	 * @see AddChild, GetChildren
+	 */
+	void ClearChildren()
+	{
+		Children.Reset();
+	}
+
+	/**
+	 * Gets the child items.
+	 *
+	 * @return Child items.
+	 * @see AddChild, ClearChildren, GetParent
+	 */
+	const TArray<TSharedPtr<FSessionBrowserTreeItem>>& GetChildren()
+	{
+		return Children;
+	}
+
+	/**
+	 * Gets the parent item.
+	 *
+	 * @return Parent item.
+	 * @see GetChildren, SetParent
+	 */
+	const TSharedPtr<FSessionBrowserTreeItem>& GetParent()
+	{
+		return Parent;
+	}
+
+	/**
+	 * Sets the parent item.
+	 *
+	 * @param Node The parent item to set.
+	 * @see GetParent
+	 */
+	void SetParent(const TSharedPtr<FSessionBrowserTreeItem>& Node)
+	{
+		Parent = Node;
+	}
+
+public:
+
+	virtual ESessionBrowserTreeNodeType GetType() = 0;
+
+private:
+
+	/** Holds the child process items. */
+	TArray<TSharedPtr<FSessionBrowserTreeItem>> Children;
+
+	/** Holds a pointer to the parent item. */
+	TSharedPtr<FSessionBrowserTreeItem> Parent;
+};
+
+
+/**
+ * Implements a group item in the session tree view.
+ */
+class FSessionBrowserGroupTreeItem
+	: public FSessionBrowserTreeItem
+{
+public:
+
+	/** Creates and initializes a new instance. */
+	FSessionBrowserGroupTreeItem(const FText& InGroupName, const FText& InToolTipText)
+		: GroupName(InGroupName)
+		, ToolTipText(InToolTipText)
+	{ }
+
+public:
+
+	/**
+	 * Gets the name of the group associated with this item.
+	 *
+	 * @return Group name.
+	 */
+	FText GetGroupName() const
+	{
+		return GroupName;
+	}
+
+	/**
+	 * Gets the tool tip text of the group associated with this item.
+	 *
+	 * @return Tool tip text.
+	 */
+	FText GetToolTipText() const
+	{
+		return ToolTipText;
+	}
+
+public:
+
+	// FSessionBrowserTreeNode interface
+
+	virtual ESessionBrowserTreeNodeType GetType() override
+	{
+		return ESessionBrowserTreeNodeType::Group;
+	}
+
+private:
+
+	/** The name of the group associated with this item. */
+	FText GroupName;
+
+	/** The tool tip text. */
+	FText ToolTipText;
+};
+
+
+/**
+ * Implements a instance item in the session tree view.
+ */
+class FSessionBrowserInstanceTreeItem
+	: public FSessionBrowserTreeItem
+{
+public:
+
+	/** Creates and initializes a new instance. */
+	FSessionBrowserInstanceTreeItem(const TSharedRef<ISessionInstanceInfo>& InInstanceInfo)
+		: InstanceInfo(InInstanceInfo)
+	{ }
+
+public:
+
+	/**
+	 * Gets the instance info associated with this item.
+	 *
+	 * @return Instance info.
+	 */
+	TSharedPtr<ISessionInstanceInfo> GetInstanceInfo() const
+	{
+		return InstanceInfo.Pin();
+	}
+
+public:
+
+	// FSessionBrowserTreeNode interface
+
+	virtual ESessionBrowserTreeNodeType GetType() override
+	{
+		return ESessionBrowserTreeNodeType::Instance;
+	}
+
+private:
+
+	/** Weak pointer to the instance info associated with this item. */
+	TWeakPtr<ISessionInstanceInfo> InstanceInfo;
+};
+
+
+/**
+ * Implements a session item in the session tree view.
+ */
+class FSessionBrowserSessionTreeItem
+	: public FSessionBrowserTreeItem
+{
+public:
+
+	/** Creates and initializes a new instance. */
+	FSessionBrowserSessionTreeItem(const TSharedRef<ISessionInfo>& InSessionInfo)
+		: SessionInfo(InSessionInfo)
+	{ }
+
+public:
+
+	/**
+	 * Gets the session info associated with this item.
+	 *
+	 * @return Session info.
+	 */
+	TSharedPtr<ISessionInfo> GetSessionInfo() const
+	{
+		return SessionInfo.Pin();
+	}
+
+public:
+
+	// FSessionBrowserTreeNode interface
+
+	virtual ESessionBrowserTreeNodeType GetType() override
+	{
+		return ESessionBrowserTreeNodeType::Session;
+	}
+
+private:
+
+	/** Weak pointer to the session info associated with this item. */
+	TWeakPtr<ISessionInfo> SessionInfo;
+};
diff --git a/Engine/Source/Developer/SessionFrontend/Private/SessionFrontendPrivatePCH.h b/Engine/Source/Developer/SessionFrontend/Private/SessionFrontendPrivatePCH.h
--- a/Engine/Source/Developer/SessionFrontend/Private/SessionFrontendPrivatePCH.h
+++ b/Engine/Source/Developer/SessionFrontend/Private/SessionFrontendPrivatePCH.h
@@ -34,10 +34,11 @@
 
 // session browser
 #include "SessionBrowserOwnerFilter.h"
+#include "SessionBrowserTreeItems.h"
 
-#include "SSessionBrowserInstanceListRow.h"
-#include "SSessionBrowserInstanceList.h"
-#include "SSessionBrowserSessionListRow.h"
+#include "SSessionBrowserTreeGroupRow.h"
+#include "SSessionBrowserTreeInstanceRow.h"
+#include "SSessionBrowserTreeSessionRow.h"
 #include "SSessionBrowserCommandBar.h"
 #include "SSessionBrowser.h"
 
diff --git a/Engine/Source/Developer/SessionFrontend/Private/Widgets/Browser/SSessionBrowser.cpp b/Engine/Source/Developer/SessionFrontend/Private/Widgets/Browser/SSessionBrowser.cpp
--- a/Engine/Source/Developer/SessionFrontend/Private/Widgets/Browser/SSessionBrowser.cpp
+++ b/Engine/Source/Developer/SessionFrontend/Private/Widgets/Browser/SSessionBrowser.cpp
@@ -13,6 +13,7 @@ SSessionBrowser::~SSessionBrowser()
 {
 	if (SessionManager.IsValid())
 	{
+		SessionManager->OnInstanceSelectionChanged().RemoveAll(this);
 		SessionManager->OnSelectedSessionChanged().RemoveAll(this);
 		SessionManager->OnSessionsUpdated().RemoveAll(this);
 	}
@@ -22,8 +23,11 @@ SSessionBrowser::~SSessionBrowser()
 /* SSessionBrowser interface
  *****************************************************************************/
 
+BEGIN_SLATE_FUNCTION_BUILD_OPTIMIZATION
 void SSessionBrowser::Construct( const FArguments& InArgs, ISessionManagerRef InSessionManager )
 {
+	IgnoreSessionManagerEvents = false;
+	IgnoreSessionTreeEvents = false;
 	SessionManager = InSessionManager;
 
 	ChildSlot
@@ -31,43 +35,51 @@ void SSessionBrowser::Construct( const FArguments& InArgs, ISessionManagerRef In
 		SNew(SVerticalBox)
 
 		+ SVerticalBox::Slot()
-			.AutoHeight()
+			.FillHeight(1.0f)
 			[
-				SNew(SHorizontalBox)
-
-				+ SHorizontalBox::Slot()
-					.AutoWidth()
-					[
-						// session combo box
-						SAssignNew(SessionComboBox, SComboBox<ISessionInfoPtr>)
-							.ContentPadding(FMargin(6.0f, 2.0f))
-							.OptionsSource(&SessionList)
-							.ToolTipText(LOCTEXT("SessionComboBoxToolTip", "Select the game session to interact with."))
-							.OnGenerateWidget(this, &SSessionBrowser::HandleSessionComboBoxGenerateWidget)
-							.OnSelectionChanged(this, &SSessionBrowser::HandleSessionComboBoxSelectionChanged)
-							[
-								SNew(STextBlock)
-									.Font(FSlateFontInfo(FPaths::EngineContentDir() / TEXT("Slate/Fonts/Roboto-Bold.ttf"), 13))
-									.Text(this, &SSessionBrowser::HandleSessionComboBoxText)
-							]
-					]
-
-				+ SHorizontalBox::Slot()
-					.FillWidth(1.0f)
-					.Padding(4.0f, 0.0f, 0.0f, 0.0f)
-					.VAlign(VAlign_Center)
-					[
-						SNew(STextBlock)
-							.Text(this, &SSessionBrowser::HandleSessionDetailsText)
-					]
-
-				+ SHorizontalBox::Slot()
-					.AutoWidth()
-					.Padding(8.0f, 0.0f)
+				// session tree
+				SNew(SBorder)
+					.BorderImage(FEditorStyle::GetBrush("ToolPanel.GroupBorder"))
+					.Padding(0.0f)
 					[
-						SNew(SSeparator)
-							.Orientation(Orient_Vertical)
+						SAssignNew(SessionTreeView, STreeView<TSharedPtr<FSessionBrowserTreeItem>>)
+							.ItemHeight(20.0f)
+							.OnExpansionChanged(this, &SSessionBrowser::HandleSessionTreeViewExpansionChanged)
+							.OnGenerateRow(this, &SSessionBrowser::HandleSessionTreeViewGenerateRow)
+							.OnGetChildren(this, &SSessionBrowser::HandleSessionTreeViewGetChildren)
+							.OnSelectionChanged(this, &SSessionBrowser::HandleSessionTreeViewSelectionChanged)
+							.SelectionMode(ESelectionMode::Multi)
+							.TreeItemsSource(&SessionTreeItems)
+							.HeaderRow
+							(
+								SNew(SHeaderRow)
+
+								+ SHeaderRow::Column("Name")
+									.DefaultLabel(LOCTEXT("InstanceListNameColumnHeader", "Name"))
+									.FillWidth(0.3f)
+
+								+ SHeaderRow::Column("Type")
+									.DefaultLabel(LOCTEXT("InstanceListTypeColumnHeader", "Type"))
+									.FillWidth(0.2f)
+
+								+ SHeaderRow::Column("Device")
+									.DefaultLabel(LOCTEXT("InstanceListDeviceColumnHeader", "Device"))
+									.FillWidth(0.3f)
+
+								+ SHeaderRow::Column("Status")
+									.DefaultLabel(LOCTEXT("InstanceListStatusColumnHeader", "Status"))
+									.FillWidth(0.2f)
+									.HAlignCell(HAlign_Right)
+									.HAlignHeader(HAlign_Right)
+							)
 					]
+			]
+		/*
+		+ SVerticalBox::Slot()
+			.AutoHeight()
+			.Padding(0.0f, 4.0f, 0.0f, 0.0f)
+			[
+				SNew(SHorizontalBox)
 
 				+ SHorizontalBox::Slot()
 					.AutoWidth()
@@ -101,89 +113,136 @@ void SSessionBrowser::Construct( const FArguments& InArgs, ISessionManagerRef In
 									]
 							]
 					]
-			]
-
-		+ SVerticalBox::Slot()
-			.FillHeight(1.0f)
-			.Padding(0.0f, 4.0f, 0.0f, 0.0f)
-			[
-				SNew(SBorder)
-					.BorderImage(FEditorStyle::GetBrush("ToolPanel.GroupBorder"))
-					.Padding(0.0f)
-					[
-						// instance list
-						SAssignNew(InstanceListView, SSessionInstanceList, SessionManager.ToSharedRef())
-							.SelectionMode(ESelectionMode::Multi)
-					]
-			]
+			]*/
 	];
 
+	AppGroupItem = MakeShareable(new FSessionBrowserGroupTreeItem(LOCTEXT("AppGroupName", "This Application"), LOCTEXT("AppGroupToolTip", "The application instance that this session browser belongs to")));
+	OtherGroupItem = MakeShareable(new FSessionBrowserGroupTreeItem(LOCTEXT("OtherGroupName", "Other Sessions"), LOCTEXT("OtherGroupToolTip", "All sessions that belong to other users")));
+	OwnerGroupItem = MakeShareable(new FSessionBrowserGroupTreeItem(LOCTEXT("OwnerGroupName", "My Sessions"), LOCTEXT("OwnerGroupToolTip", "All sessions that were started by you")));
+	StandaloneGroupItem = MakeShareable(new FSessionBrowserGroupTreeItem(LOCTEXT("StandaloneGroupName", "Standalone Instances"), LOCTEXT("StandaloneGroupToolTip", "Engine instances that don't belong to any particular session")));
+
+	SessionTreeItems.Add(AppGroupItem);
+	SessionTreeItems.Add(OwnerGroupItem);
+	SessionTreeItems.Add(OtherGroupItem);
+	SessionTreeItems.Add(StandaloneGroupItem);
+
+	SessionManager->OnInstanceSelectionChanged().AddSP(this, &SSessionBrowser::HandleSessionManagerInstanceSelectionChanged);
 	SessionManager->OnSelectedSessionChanged().AddSP(this, &SSessionBrowser::HandleSessionManagerSelectedSessionChanged);
 	SessionManager->OnSessionsUpdated().AddSP(this, &SSessionBrowser::HandleSessionManagerSessionsUpdated);
 
+	SessionTreeView->SetSingleExpandedItem(AppGroupItem);
+
 	ReloadSessions();
 }
+END_SLATE_FUNCTION_BUILD_OPTIMIZATION
 
 
 /* SSessionBrowser implementation
  *****************************************************************************/
 
-void SSessionBrowser::FilterSessions()
+void SSessionBrowser::ExpandItem(const TSharedPtr<FSessionBrowserTreeItem>& Item)
 {
-	SessionList.Reset();
+	SessionTreeView->SetSingleExpandedItem(Item);
 
-	for (int32 Index = 0; Index < AvailableSessions.Num(); ++Index)
+	if (Item.IsValid())
 	{
-		const ISessionInfoPtr& Session = AvailableSessions[Index];
+		const auto& ParentItem = Item->GetParent();
 
-		if (Session->GetSessionOwner() == FPlatformProcess::UserName(true))
+		if (ParentItem.IsValid())
 		{
-			SessionList.Add(Session);
+			SessionTreeView->SetItemExpansion(ParentItem, true);
 		}
 	}
-
-	SessionComboBox->RefreshOptions();
-	SessionComboBox->SetSelectedItem(SessionManager->GetSelectedSession());
 }
 
 
-FText SSessionBrowser::GetSessionName( const ISessionInfoPtr& SessionInfo ) const
+void SSessionBrowser::FilterSessions()
 {
-	FText SessionName;
+	// clear tree groups
+	AppGroupItem->ClearChildren();
+	OwnerGroupItem->ClearChildren();
+	OtherGroupItem->ClearChildren();
+	StandaloneGroupItem->ClearChildren();
 
-	if (!SessionInfo->IsStandalone() || !SessionInfo->GetSessionName().IsEmpty())
-	{
-		SessionName = FText::FromString(SessionInfo->GetSessionName());
-	}
-	else
+	// update tree items
+	TMap<FGuid, TSharedPtr<FSessionBrowserTreeItem>> NewItemMap;
+
+	for (auto& SessionInfo : AvailableSessions)
 	{
+		// process session
+		bool LocalOwner = SessionInfo->GetSessionOwner() == FPlatformProcess::UserName(true);
+
+		if (!SessionInfo->IsStandalone() && !LocalOwner)
+		{
+			continue;
+		}
+
+		TSharedPtr<FSessionBrowserTreeItem> SessionItem = ItemMap.FindRef(SessionInfo->GetSessionId());
+
+		if (SessionItem.IsValid())
+		{
+			SessionItem->ClearChildren();
+		}
+		else
+		{
+			SessionItem = MakeShareable(new FSessionBrowserSessionTreeItem(SessionInfo.ToSharedRef()));
+		}
+
+		NewItemMap.Add(SessionInfo->GetSessionId(), SessionItem);
+
+		// add session to group
 		TArray<ISessionInstanceInfoPtr> Instances;
 		SessionInfo->GetInstances(Instances);
 
-		if (Instances.Num() > 0)
+		if (LocalOwner)
+		{
+			if (Instances[0]->GetInstanceId() != FApp::GetInstanceId())
+			{
+				OwnerGroupItem->AddChild(SessionItem.ToSharedRef());
+				SessionItem->SetParent(OwnerGroupItem);
+			}
+		}
+		else if (SessionInfo->IsStandalone())
+		{
+			StandaloneGroupItem->AddChild(SessionItem.ToSharedRef());
+			SessionItem->SetParent(StandaloneGroupItem);
+		}
+		else
+		{
+			OtherGroupItem->AddChild(SessionItem.ToSharedRef());
+			SessionItem->SetParent(OtherGroupItem);
+		}
+		
+		// process session instances
+		for (auto& InstanceInfo : Instances)
 		{
-			const ISessionInstanceInfoPtr& FirstInstance = Instances[0];
+			TSharedPtr<FSessionBrowserTreeItem> InstanceItem = ItemMap.FindRef(InstanceInfo->GetInstanceId());
 
-			if ((Instances.Num() == 1) && (FirstInstance->GetInstanceId() == FApp::GetInstanceId()))
+			if (!InstanceItem.IsValid())
 			{
-				SessionName = LOCTEXT("ThisApplicationSessionText", "This Application");
+				InstanceItem = MakeShareable(new FSessionBrowserInstanceTreeItem(InstanceInfo.ToSharedRef()));
 			}
-			else if (FirstInstance->GetDeviceName() == FPlatformProcess::ComputerName())
+
+			NewItemMap.Add(InstanceInfo->GetInstanceId(), InstanceItem);
+
+			// add instance to group or session
+			if (InstanceInfo->GetInstanceId() == FApp::GetInstanceId())
 			{
-				SessionName = LOCTEXT("UnnamedLocalSessionText", "Unnamed Session (Local)");
+				AppGroupItem->AddChild(InstanceItem.ToSharedRef());
+				InstanceItem->SetParent(AppGroupItem);
 			}
 			else
 			{
-				SessionName = LOCTEXT("UnnamedRemoteSessionText", "Unnamed Session (Remote)");
+				InstanceItem->SetParent(SessionItem);
+				SessionItem->AddChild(InstanceItem.ToSharedRef());
 			}
 		}
-		else
-		{
-			SessionName = LOCTEXT("UnnamedSessionText", "Unnamed Session");
-		}
 	}
 
-	return SessionName;
+	ItemMap = NewItemMap;
+
+	// refresh tree view
+	SessionTreeView->RequestTreeRefresh();
 }
 
 
@@ -191,124 +250,194 @@ void SSessionBrowser::ReloadSessions()
 {
 	SessionManager->GetSessions(AvailableSessions);
 	FilterSessions();
-
-	if ((SessionList.Num() > 0) && !SessionManager->GetSelectedSession().IsValid())
-	{
-		SessionManager->SelectSession(SessionList[0]);
-	}
 }
 
 
 /* SSessionBrowser event handlers
  *****************************************************************************/
 
-TSharedRef<SWidget> SSessionBrowser::HandleSessionComboBoxGenerateWidget( ISessionInfoPtr SessionInfo ) const
+void SSessionBrowser::HandleSessionManagerInstanceSelectionChanged(const TSharedPtr<ISessionInstanceInfo>& Instance, bool Selected)
 {
-	return SNew(SVerticalBox)
-
-	+ SVerticalBox::Slot()
-		.AutoHeight()
-		.Padding(4.0f, 0.0f)
-		[
-			SNew(STextBlock)
-				.Font(FSlateFontInfo(FPaths::EngineContentDir() / TEXT("Slate/Fonts/Roboto-Bold.ttf"), 14))
-				.Text(GetSessionName(SessionInfo))
-		]
-
-	+ SVerticalBox::Slot()
-		.AutoHeight()
-		.Padding(4.0f, 0.0f)
-		[
-			SNew(SHorizontalBox)
-
-			+ SHorizontalBox::Slot()
-				.AutoWidth()
-				.HAlign(HAlign_Left)
-				[
-					SNew(STextBlock)
-						.Text(FText::FromString(SessionInfo->GetSessionOwner()))
-				]
-
-			+ SHorizontalBox::Slot()
-				.AutoWidth()
-				.HAlign(HAlign_Right)
-				.Padding(16.0f, 0.0f, 0.0f, 0.0f)
-				[
-					SNew(STextBlock)
-						.Text(FText::Format(LOCTEXT("InstanceCountFormat", "{0} Instance(s)"), FText::AsNumber(SessionInfo->GetNumInstances())))
-				]
-		];
+	if (IgnoreSessionManagerEvents)
+	{
+		return;
+	}
+
+	const auto& InstanceItem = ItemMap.FindRef(Instance->GetInstanceId());
+
+	if (InstanceItem.IsValid())
+	{
+		SessionTreeView->SetItemSelection(InstanceItem, Selected);
+	}
 }
 
 
-void SSessionBrowser::HandleSessionComboBoxSelectionChanged( ISessionInfoPtr SelectedItem, ESelectInfo::Type SelectInfo )
+void SSessionBrowser::HandleSessionManagerSelectedSessionChanged(const ISessionInfoPtr& SelectedSession)
 {
-	if (!SessionManager->SelectSession(SelectedItem))
+	if (IgnoreSessionManagerEvents)
+	{
+		return;
+	}
+
+	IgnoreSessionTreeEvents = true;
 	{
-		SessionComboBox->SetSelectedItem(SessionManager->GetSelectedSession());
+		if (SelectedSession.IsValid())
+		{
+			ExpandItem(ItemMap.FindRef(SelectedSession->GetSessionId()));
+		}
+		else
+		{
+			SessionTreeView->SetSingleExpandedItem(nullptr);
+		}
 	}
+	IgnoreSessionTreeEvents = false;
+}
+
+
+void SSessionBrowser::HandleSessionManagerSessionsUpdated()
+{
+	ReloadSessions();
 }
 
 
-FText SSessionBrowser::HandleSessionComboBoxText() const
+FText SSessionBrowser::HandleSessionTreeRowGetToolTipText(TSharedPtr<FSessionBrowserTreeItem> Item) const
 {
-	const ISessionInfoPtr& SelectedSession = SessionManager->GetSelectedSession();
+	FTextBuilder ToolTipTextBuilder;
 
-	if (SelectedSession.IsValid())
+	if (Item->GetType() == ESessionBrowserTreeNodeType::Instance)
 	{
-		return GetSessionName(SelectedSession);
+		ISessionInstanceInfoPtr InstanceInfo = StaticCastSharedPtr<FSessionBrowserInstanceTreeItem>(Item)->GetInstanceInfo();
+
+		if (InstanceInfo.IsValid())
+		{
+			ToolTipTextBuilder.AppendLineFormat(LOCTEXT("InstanceToolTipInstanceId", "Instance ID: {0}"), FText::FromString(InstanceInfo->GetInstanceId().ToString(EGuidFormats::DigitsWithHyphensInBraces)));
+			ToolTipTextBuilder.AppendLine();
+			ToolTipTextBuilder.AppendLineFormat(LOCTEXT("InstanceToolTipBuildDate", "Build Date: {0}"), FText::FromString(InstanceInfo->GetBuildDate()));
+			ToolTipTextBuilder.AppendLineFormat(LOCTEXT("InstanceToolTipConsoleBuild", "Console Build: {0}"), InstanceInfo->IsConsole() ? LOCTEXT("LabelYes", "Yes") : LOCTEXT("LabelNo", "No"));
+			ToolTipTextBuilder.AppendLineFormat(LOCTEXT("InstanceToolTipEngineVersion", "Engine Version: {0}"), InstanceInfo->GetEngineVersion() == 0 ? LOCTEXT("CustomBuildVersion", "Custom Build") : FText::FromString(FString::FromInt(InstanceInfo->GetEngineVersion())));
+			ToolTipTextBuilder.AppendLineFormat(LOCTEXT("InstanceToolTipPlatform", "Platform: {0}"), FText::FromString(InstanceInfo->GetPlatformName()));
+			ToolTipTextBuilder.AppendLineFormat(LOCTEXT("InstanceToolTipCurrentLevel", "Current Level: {0}"), FText::FromString(InstanceInfo->GetCurrentLevel()));
+			ToolTipTextBuilder.AppendLineFormat(LOCTEXT("InstanceToolTipWorldTimeSeconds", "World Time: {0}"), FText::AsTimespan(FTimespan::FromSeconds(InstanceInfo->GetWorldTimeSeconds())));
+			ToolTipTextBuilder.AppendLineFormat(LOCTEXT("InstanceToolTipPlayBegun", "Play Has Begun: {0}"), InstanceInfo->PlayHasBegun() ? GYes : GNo);
+			ToolTipTextBuilder.AppendLine();
+			ToolTipTextBuilder.AppendLineFormat(LOCTEXT("SessionToolTipLastUpdateTime", "Last Update Time: {0}"), FText::AsDateTime(InstanceInfo->GetLastUpdateTime()));
+		}
 	}
 
-	return LOCTEXT("SelectSessionHint", "Select a session...");
+	return ToolTipTextBuilder.ToText();
 }
 
 
-FText SSessionBrowser::HandleSessionDetailsText() const
+void SSessionBrowser::HandleSessionTreeViewExpansionChanged(TSharedPtr<FSessionBrowserTreeItem> TreeItem, bool bIsExpanded)
 {
-	const ISessionInfoPtr& SelectedSession = SessionManager->GetSelectedSession();
+	if (IgnoreSessionTreeEvents || !TreeItem.IsValid())
+	{
+		return;
+	}
 
-	if (!SelectedSession.IsValid())
+	if (TreeItem->GetType() == ESessionBrowserTreeNodeType::Instance)
 	{
-		return FText::GetEmpty();
+		return;
 	}
 
-	FFormatNamedArguments FormatArguments;
-	FormatArguments.Add(TEXT("NumInstances"), SelectedSession->GetNumInstances());
-	FormatArguments.Add(TEXT("OwnerName"), SelectedSession->GetSessionOwner().IsEmpty() ? LOCTEXT("UnknownOwner", "<unknown>") : FText::FromString(SelectedSession->GetSessionOwner()));
+	IgnoreSessionManagerEvents = true;
+	{
+		if (bIsExpanded)
+		{
+			IgnoreSessionTreeEvents = true;
+			{
+				ExpandItem(TreeItem);
+			}
+			IgnoreSessionTreeEvents = false;
 
-	return FText::Format(LOCTEXT("SessionDetails", "Owner: {OwnerName}, {NumInstances} Instance(s)"), FormatArguments);
+			// select session
+			if (TreeItem->GetType() == ESessionBrowserTreeNodeType::Session)
+			{
+				SessionManager->SelectSession(StaticCastSharedPtr<FSessionBrowserSessionTreeItem>(TreeItem)->GetSessionInfo());
+			}
+			else
+			{
+				SessionManager->SelectSession(nullptr);
+			}
+		}
+		else if (TreeItem->GetType() != ESessionBrowserTreeNodeType::Instance)
+		{
+			// deselect session
+			SessionManager->SelectSession(nullptr);
+		}
+	}
+	IgnoreSessionManagerEvents = false;
 }
 
 
-TSharedRef<ITableRow> SSessionBrowser::HandleSessionListViewGenerateRow( ISessionInfoPtr Item, const TSharedRef<STableViewBase>& OwnerTable )
+TSharedRef<ITableRow> SSessionBrowser::HandleSessionTreeViewGenerateRow(TSharedPtr<FSessionBrowserTreeItem> Item, const TSharedRef<STableViewBase>& OwnerTable)
 {
-	return SNew(SSessionBrowserSessionListRow, OwnerTable)
-		.SessionInfo(Item)
-		.ToolTipText(FText::Format(LOCTEXT("SessionToolTipSessionIdFmt", "Session ID: {0}"), FText::FromString(Item->GetSessionId().ToString(EGuidFormats::DigitsWithHyphensInBraces))));
+	if (Item->GetType() == ESessionBrowserTreeNodeType::Group)
+	{
+		return SNew(SSessionBrowserTreeGroupRow, OwnerTable)
+			.Item(StaticCastSharedRef<FSessionBrowserGroupTreeItem>(Item.ToSharedRef()));
+	}
+
+	if (Item->GetType() == ESessionBrowserTreeNodeType::Session)
+	{
+		return SNew(SSessionBrowserTreeSessionRow, OwnerTable)
+			.Item(StaticCastSharedRef<FSessionBrowserSessionTreeItem>(Item.ToSharedRef()));
+	}
+
+	return SNew(SSessionBrowserTreeInstanceRow, OwnerTable)
+		.Item(StaticCastSharedRef<FSessionBrowserInstanceTreeItem>(Item.ToSharedRef()))
+		.ToolTipText(this, &SSessionBrowser::HandleSessionTreeRowGetToolTipText, Item);
 }
 
 
-void SSessionBrowser::HandleSessionListViewSelectionChanged( ISessionInfoPtr Item, ESelectInfo::Type SelectInfo )
+void SSessionBrowser::HandleSessionTreeViewGetChildren(TSharedPtr<FSessionBrowserTreeItem> Item, TArray<TSharedPtr<FSessionBrowserTreeItem>>& OutChildren)
 {
-	if (!SessionManager->SelectSession(Item))
+	if (Item.IsValid())
 	{
-		SessionListView->SetSelection(SessionManager->GetSelectedSession());
+		OutChildren = Item->GetChildren();
 	}
 }
 
 
-void SSessionBrowser::HandleSessionManagerSelectedSessionChanged( const ISessionInfoPtr& SelectedSession )
+void SSessionBrowser::HandleSessionTreeViewSelectionChanged(const TSharedPtr<FSessionBrowserTreeItem> Item, ESelectInfo::Type SelectInfo)
 {
-	if (SessionComboBox->GetSelectedItem() != SelectedSession)
+	if (IgnoreSessionTreeEvents || (SelectInfo == ESelectInfo::Direct))
 	{
-		SessionComboBox->SetSelectedItem(SelectedSession);
+		return;
 	}
-}
 
+	IgnoreSessionManagerEvents = true;
+	{
+		if (Item.IsValid())
+		{
+			const auto& InstanceInfo = StaticCastSharedPtr<FSessionBrowserInstanceTreeItem>(Item)->GetInstanceInfo();
 
-void SSessionBrowser::HandleSessionManagerSessionsUpdated()
-{
-	ReloadSessions();
+			if (InstanceInfo.IsValid())
+			{
+				// special handling for local application
+				if (Item->GetParent() == AppGroupItem)
+				{
+					SessionManager->SelectSession(InstanceInfo->GetOwnerSession());
+				}
+
+				SessionManager->SetInstanceSelected(InstanceInfo.ToSharedRef(), true);
+			}
+		}
+		else
+		{
+			// an instance got deselected
+			for (const auto& InstanceInfo : SessionManager->GetSelectedInstances())
+			{
+				const auto& InstanceItem = ItemMap.FindRef(InstanceInfo->GetInstanceId());
+
+				if (!SessionTreeView->IsItemSelected(InstanceItem))
+				{
+					SessionManager->SetInstanceSelected(InstanceInfo.ToSharedRef(), false);
+				}
+			}
+		}
+	}
+	IgnoreSessionManagerEvents = false;
 }
 
 
diff --git a/Engine/Source/Developer/SessionFrontend/Private/Widgets/Browser/SSessionBrowser.h b/Engine/Source/Developer/SessionFrontend/Private/Widgets/Browser/SSessionBrowser.h
--- a/Engine/Source/Developer/SessionFrontend/Private/Widgets/Browser/SSessionBrowser.h
+++ b/Engine/Source/Developer/SessionFrontend/Private/Widgets/Browser/SSessionBrowser.h
@@ -31,40 +31,45 @@ public:
 
 protected:
 
-	/** Filters the session list. */
-	void FilterSessions();
-
 	/**
-	 * Gets the display name of the specified session.
+	 * Fully expands the specified tree view item.
 	 *
-	 * @param SessionInfo The session to get the name for.
-	 * @param Session name.
+	 * @param Item The tree view item to expand.
 	 */
-	FText GetSessionName( const ISessionInfoPtr& SessionInfo ) const;
+	void ExpandItem(const TSharedPtr<FSessionBrowserTreeItem>& Item);
+
+	/** Filters the session tree. */
+	void FilterSessions();
 
 	/** Reloads the sessions list. */
 	void ReloadSessions();
 
 private:
 
-	/** Callback for generating a row in the session combo box. */
-	TSharedRef<SWidget> HandleSessionComboBoxGenerateWidget( ISessionInfoPtr SessionInfo ) const;
-
-	/** Callback for selection changes in the session combo box. */
-	void HandleSessionComboBoxSelectionChanged( ISessionInfoPtr SelectedItem, ESelectInfo::Type SelectInfo );
-
-	/** Callback for getting the text in the session combo box. */
-	FText HandleSessionComboBoxText() const;
-
-	/** Callback for getting the text of the session details area. */
-	FText HandleSessionDetailsText() const;
+	/** Callback for changing the selection state of an instance. */
+	void HandleSessionManagerInstanceSelectionChanged(const TSharedPtr<ISessionInstanceInfo>& Instance, bool Selected);
 
 	/** Callback for changing the selected session in the session manager. */
-	void HandleSessionManagerSelectedSessionChanged( const ISessionInfoPtr& SelectedSession );
+	void HandleSessionManagerSelectedSessionChanged(const ISessionInfoPtr& SelectedSession);
 
 	/** Callback for updating the session list in the session manager. */
 	void HandleSessionManagerSessionsUpdated();
 
+	/** Callback for getting the tool tip text of a session tree row. */
+	FText HandleSessionTreeRowGetToolTipText(TSharedPtr<FSessionBrowserTreeItem> Item) const;
+
+	/** Callback for session tree view selection changes. */
+	void HandleSessionTreeViewExpansionChanged(TSharedPtr<FSessionBrowserTreeItem> TreeItem, bool bIsExpanded);
+
+	/** Callback for generating a row widget in the session tree view. */
+	TSharedRef<ITableRow> HandleSessionTreeViewGenerateRow(TSharedPtr<FSessionBrowserTreeItem> Item, const TSharedRef<STableViewBase>& OwnerTable);
+
+	/** Callback for getting the children of a node in the session tree view. */
+	void HandleSessionTreeViewGetChildren(TSharedPtr<FSessionBrowserTreeItem> Item, TArray<TSharedPtr<FSessionBrowserTreeItem>>& OutChildren);
+
+	/** Callback for session tree view selection changes. */
+	void HandleSessionTreeViewSelectionChanged(const TSharedPtr<FSessionBrowserTreeItem> Item, ESelectInfo::Type SelectInfo);
+
 	/** Callback for clicking the 'Terminate Session' button. */
 	FReply HandleTerminateSessionButtonClicked();
 
@@ -79,26 +84,35 @@ private:
 	/** Holds the command bar. */
 	TSharedPtr<SSessionBrowserCommandBar> CommandBar;
 
-	/** Holds the instance list view. */
-	TSharedPtr<SSessionInstanceList> InstanceListView;
+	/** Whether to ignore events from the session manager. */
+	bool IgnoreSessionManagerEvents;
 
-	/** Holds the session combo box. */
-	TSharedPtr<SComboBox<ISessionInfoPtr>> SessionComboBox;
+	/** Whether to ignore events from the session tree view. */
+	bool IgnoreSessionTreeEvents;
 
-	/** Holds the filtered list of sessions to be displayed. */
-	TArray<ISessionInfoPtr> SessionList;
+	/** Maps session and instance GUIDs to existing tree items. */
+	TMap<FGuid, TSharedPtr<FSessionBrowserTreeItem>> ItemMap;
 
 	/** Holds a reference to the session manager. */
 	ISessionManagerPtr SessionManager;
 
+	/** Holds the filtered list of tree items. */
+	TArray<TSharedPtr<FSessionBrowserTreeItem>> SessionTreeItems;
+
+	/** Holds the session tree view. */
+	TSharedPtr<STreeView<TSharedPtr<FSessionBrowserTreeItem>>> SessionTreeView;
+
 private:
 
-	/** Callback for creating a row widget in the session list view. */
-	TSharedRef<ITableRow> HandleSessionListViewGenerateRow( ISessionInfoPtr Item, const TSharedRef<STableViewBase>& OwnerTable );
+	/** The session tree item that holds this application's session. */
+	TSharedPtr<FSessionBrowserGroupTreeItem> AppGroupItem;
+
+	/** The session tree item that holds other user's sessions. */
+	TSharedPtr<FSessionBrowserGroupTreeItem> OtherGroupItem;
 
-	/** Handles changing the selection in the session tree. */
-	void HandleSessionListViewSelectionChanged( ISessionInfoPtr Item, ESelectInfo::Type SelectInfo );
+	/** The session tree item that holds the owner's remote sessions. */
+	TSharedPtr<FSessionBrowserGroupTreeItem> OwnerGroupItem;
 
-	/** Holds the session list view. */
-	TSharedPtr<SListView<ISessionInfoPtr>> SessionListView;
+	/** The session tree item that holds other user's standalone instances. */
+	TSharedPtr<FSessionBrowserGroupTreeItem> StandaloneGroupItem;
 };
diff --git a/Engine/Source/Developer/SessionFrontend/Private/Widgets/Browser/SSessionBrowserInstanceList.h b/Engine/Source/Developer/SessionFrontend/Private/Widgets/Browser/SSessionBrowserInstanceList.h
deleted file mode 100644
diff --git a/Engine/Source/Developer/SessionFrontend/Private/Widgets/Browser/SSessionBrowserInstanceListRow.h b/Engine/Source/Developer/SessionFrontend/Private/Widgets/Browser/SSessionBrowserInstanceListRow.h
deleted file mode 100644
diff --git a/Engine/Source/Developer/SessionFrontend/Private/Widgets/Browser/SSessionBrowserSessionListRow.h b/Engine/Source/Developer/SessionFrontend/Private/Widgets/Browser/SSessionBrowserSessionListRow.h
deleted file mode 100644
diff --git a/Engine/Source/Developer/SessionFrontend/Private/Widgets/Browser/SSessionBrowserTreeGroupRow.h b/Engine/Source/Developer/SessionFrontend/Private/Widgets/Browser/SSessionBrowserTreeGroupRow.h
new file mode 100644
--- /dev/null
+++ b/Engine/Source/Developer/SessionFrontend/Private/Widgets/Browser/SSessionBrowserTreeGroupRow.h
@@ -0,0 +1,140 @@
+// Copyright 1998-2015 Epic Games, Inc. All Rights Reserved.
+
+#pragma once
+
+
+#define LOCTEXT_NAMESPACE "SSessionBrowserTreeRow"
+
+
+/**
+ * Implements a row widget for the session browser tree.
+ */
+class SSessionBrowserTreeGroupRow
+	: public STableRow<TSharedPtr<FSessionBrowserGroupTreeItem>>
+{
+public:
+
+	SLATE_BEGIN_ARGS(SSessionBrowserTreeGroupRow) { }
+		SLATE_ATTRIBUTE(FText, HighlightText)
+		SLATE_ARGUMENT(TSharedPtr<STableViewBase>, OwnerTableView)
+		SLATE_ARGUMENT(TSharedPtr<FSessionBrowserGroupTreeItem>, Item)
+	SLATE_END_ARGS()
+
+public:
+
+	/**
+	 * Constructs the widget.
+	 *
+	 * @param InArgs The construction arguments.
+	 */
+	BEGIN_SLATE_FUNCTION_BUILD_OPTIMIZATION
+	void Construct(const FArguments& InArgs, const TSharedRef<STableViewBase>& InOwnerTableView)
+	{
+		HighlightText = InArgs._HighlightText;
+		Item = InArgs._Item;
+
+		ChildSlot
+			.Padding(0.0f, 2.0f, 0.0f, 0.0f)
+			[
+				SNew(SBorder)
+					.BorderBackgroundColor(FLinearColor(0.6f, 0.6f, 0.6f, 1.0f))
+					.BorderImage(this, &SSessionBrowserTreeGroupRow::HandleBorderBackgroundImage)
+					.Padding(3.0f)
+					.ToolTipText(Item->GetToolTipText())
+					[
+						SNew(SHorizontalBox)
+
+						+ SHorizontalBox::Slot()
+							.FillWidth(1.0f)
+							.VAlign(VAlign_Center)
+							[
+								SNew(STextBlock)
+									.Text(Item->GetGroupName())
+									.Font(FEditorStyle::GetFontStyle("DetailsView.CategoryFontStyle"))
+									.ShadowOffset(FVector2D(1.0f, 1.0f))
+							]
+
+						+ SHorizontalBox::Slot()
+							.AutoWidth()
+							.VAlign(VAlign_Center)
+							[
+								SNew(SImage)
+									.Image(this, &SSessionBrowserTreeGroupRow::HandlePullDownImage)
+							]
+					]
+			];
+
+		STableRow<TSharedPtr<FSessionBrowserGroupTreeItem>>::ConstructInternal(
+			STableRow<TSharedPtr<FSessionBrowserGroupTreeItem>>::FArguments()
+				.ShowSelection(false)
+				.Style(FEditorStyle::Get(), "DetailsView.TreeView.TableRow"),
+			InOwnerTableView
+		);
+	}
+	END_SLATE_FUNCTION_BUILD_OPTIMIZATION
+
+protected:
+
+	// SWidget overrides
+
+	virtual FReply OnMouseButtonDown(const FGeometry& MyGeometry, const FPointerEvent& MouseEvent) override
+	{
+		if (MouseEvent.GetEffectingButton() == EKeys::LeftMouseButton)
+		{
+			ToggleExpansion();
+
+			return FReply::Handled();
+		}
+		else
+		{
+			return FReply::Unhandled();
+		}
+	}
+
+private:
+
+	/** Callback for getting the background image of the row's border. */
+	const FSlateBrush* HandleBorderBackgroundImage() const
+	{
+		if (IsHovered())
+		{
+			return IsItemExpanded()
+				? FEditorStyle::GetBrush("DetailsView.CategoryTop_Hovered")
+				: FEditorStyle::GetBrush("DetailsView.CollapsedCategory_Hovered");
+		}
+		else
+		{
+			return IsItemExpanded()
+				? FEditorStyle::GetBrush("DetailsView.CategoryTop")
+				: FEditorStyle::GetBrush("DetailsView.CollapsedCategory");
+		}
+	}
+
+	/** Gets the image for the pull-down icon. */
+	const FSlateBrush* HandlePullDownImage() const
+	{
+		if (IsHovered())
+		{
+			return IsItemExpanded()
+				? FEditorStyle::GetBrush("DetailsView.PulldownArrow.Up.Hovered")
+				: FEditorStyle::GetBrush("DetailsView.PulldownArrow.Down.Hovered");
+		}
+		else
+		{
+			return IsItemExpanded()
+				? FEditorStyle::GetBrush("DetailsView.PulldownArrow.Up")
+				: FEditorStyle::GetBrush("DetailsView.PulldownArrow.Down");
+		}
+	}
+
+private:
+
+	/** The highlight string for this row. */
+	TAttribute<FText> HighlightText;
+
+	/** A reference to the tree item that is displayed in this row. */
+	TSharedPtr<FSessionBrowserGroupTreeItem> Item;
+};
+
+
+#undef LOCTEXT_NAMESPACE
diff --git a/Engine/Source/Developer/SessionFrontend/Private/Widgets/Browser/SSessionBrowserTreeInstanceRow.h b/Engine/Source/Developer/SessionFrontend/Private/Widgets/Browser/SSessionBrowserTreeInstanceRow.h
new file mode 100644
--- /dev/null
+++ b/Engine/Source/Developer/SessionFrontend/Private/Widgets/Browser/SSessionBrowserTreeInstanceRow.h
@@ -0,0 +1,303 @@
+// Copyright 1998-2015 Epic Games, Inc. All Rights Reserved.
+
+#pragma once
+
+
+#define LOCTEXT_NAMESPACE "SSessionBrowserTreeRow"
+
+
+/**
+ * Implements a row widget for the session browser tree.
+ */
+class SSessionBrowserTreeInstanceRow
+	: public SMultiColumnTableRow<TSharedPtr<FSessionBrowserTreeItem>>
+{
+public:
+
+	SLATE_BEGIN_ARGS(SSessionBrowserTreeInstanceRow) { }
+		SLATE_ATTRIBUTE(FText, HighlightText)
+		SLATE_ARGUMENT(TSharedPtr<STableViewBase>, OwnerTableView)
+		SLATE_ARGUMENT(TSharedPtr<FSessionBrowserInstanceTreeItem>, Item)
+		SLATE_ARGUMENT(bool, ShowSelection)
+	SLATE_END_ARGS()
+
+public:
+
+	/**
+	 * Constructs the widget.
+	 *
+	 * @param InArgs The construction arguments.
+	 */
+	void Construct(const FArguments& InArgs, const TSharedRef<STableViewBase>& InOwnerTableView)
+	{
+		HighlightText = InArgs._HighlightText;
+		Item = InArgs._Item;
+
+		SMultiColumnTableRow<TSharedPtr<FSessionBrowserTreeItem>>::Construct(FSuperRowType::FArguments().ShowSelection(InArgs._ShowSelection).Style(FEditorStyle::Get(), "TableView.Row"), InOwnerTableView);
+	}
+
+public:
+
+	/**
+	 * Generates the widget for the specified column.
+	 *
+	 * @param ColumnName The name of the column to generate the widget for.
+	 * @return The widget.
+	 */
+	BEGIN_SLATE_FUNCTION_BUILD_OPTIMIZATION
+	virtual TSharedRef<SWidget> GenerateWidgetForColumn( const FName& ColumnName ) override
+	{
+		if (ColumnName == "Device")
+		{
+			ISessionInstanceInfoPtr InstanceInfo = Item->GetInstanceInfo();
+
+			if (InstanceInfo.IsValid())
+			{
+				const PlatformInfo::FPlatformInfo* const PlatformInfo = PlatformInfo::FindPlatformInfo(*InstanceInfo->GetPlatformName());
+
+				return SNew(SHorizontalBox)
+
+				+ SHorizontalBox::Slot()
+					.AutoWidth()
+					.VAlign(VAlign_Center)
+					[
+						SNew(SBox)
+							.WidthOverride(24)
+							.HeightOverride(24)
+							[
+								SNew(SImage)
+									.Image((PlatformInfo) ? FEditorStyle::GetBrush(PlatformInfo->GetIconStyleName(PlatformInfo::EPlatformIconSize::Normal)) : FStyleDefaults::GetNoBrush())
+							]
+					]
+
+				+ SHorizontalBox::Slot()
+					.FillWidth(1.0f)
+					.Padding(2.0f, 0.0f, 0.0f, 0.0f)
+					.VAlign(VAlign_Center)
+					[
+						SNew(STextBlock)
+							.ColorAndOpacity(this, &SSessionBrowserTreeInstanceRow::HandleTextColorAndOpacity)
+							.Text(this, &SSessionBrowserTreeInstanceRow::HandleDeviceColumnText)
+					];
+			}
+		}
+		else if (ColumnName == "Level")
+		{
+			return SNew(SBox)
+				.Padding(FMargin(4.0, 0.0))
+				.VAlign(VAlign_Center)
+				[
+					SNew(STextBlock)
+						.ColorAndOpacity(this, &SSessionBrowserTreeInstanceRow::HandleTextColorAndOpacity)
+						.Text(this, &SSessionBrowserTreeInstanceRow::HandleLevelColumnText)
+				];
+		}
+		else if (ColumnName == "Name")
+		{
+			ISessionInstanceInfoPtr InstanceInfo = Item->GetInstanceInfo();
+
+			if (InstanceInfo.IsValid())
+			{
+				return SNew(SBox)
+					.Padding(FMargin(1.0f, 1.0f, 4.0f, 1.0f))
+					.HAlign(HAlign_Left)
+					[
+						SNew(SBorder)
+							.BorderBackgroundColor(this, &SSessionBrowserTreeInstanceRow::HandleInstanceBorderBackgroundColor)
+							.BorderImage(this, &SSessionBrowserTreeInstanceRow::HandleInstanceBorderBrush)
+							.ColorAndOpacity(FLinearColor(0.25f, 0.25f, 0.25f))
+							.Padding(FMargin(6.0f, 4.0f))
+							.VAlign(VAlign_Center)
+							[
+								SNew(STextBlock)
+									.Font(FEditorStyle::GetFontStyle("BoldFont"))
+									.Text(FText::FromString(InstanceInfo->GetInstanceName()))
+							]
+					];
+			}
+		}
+		else if (ColumnName == "Status")
+		{
+			return SNew(SBox)
+				.Padding(FMargin(4.0f, 0.0f))
+				.VAlign(VAlign_Center)
+				[
+					SNew(STextBlock)
+						.ColorAndOpacity(this, &SSessionBrowserTreeInstanceRow::HandleTextColorAndOpacity)
+						.Text(this, &SSessionBrowserTreeInstanceRow::HandleStatusColumnText)
+				];
+		}
+		else if (ColumnName == "Type")
+		{
+			ISessionInstanceInfoPtr InstanceInfo = Item->GetInstanceInfo();
+
+			if (InstanceInfo.IsValid())
+			{
+				return SNew(SBox)
+					.Padding(FMargin(4.0f, 0.0f))
+					.VAlign(VAlign_Center)
+					[
+						SNew(STextBlock)
+							.ColorAndOpacity(this, &SSessionBrowserTreeInstanceRow::HandleTextColorAndOpacity)
+							.Text(this, &SSessionBrowserTreeInstanceRow::HandleInstanceTypeText)
+					];
+			}
+		}
+
+		return SNullWidget::NullWidget;
+	}
+	END_SLATE_FUNCTION_BUILD_OPTIMIZATION
+
+protected:
+
+	FText SanitizeSessionName(const ISessionInfoPtr& SessionInfo) const
+	{
+		const FString& SessionName = SessionInfo->GetSessionName();
+
+		// return name of a launched session
+		if (!SessionInfo->IsStandalone() || !SessionName.IsEmpty())
+		{
+			return FText::FromString(SessionName);
+		}
+
+		// generate name for a standalone session
+		TArray<ISessionInstanceInfoPtr> Instances;
+		SessionInfo->GetInstances(Instances);
+
+		if (Instances.Num() > 0)
+		{
+			const ISessionInstanceInfoPtr& FirstInstance = Instances[0];
+
+			if ((Instances.Num() == 1) && (FirstInstance->GetInstanceId() == FApp::GetInstanceId()))
+			{
+				return LOCTEXT("ThisApplicationSessionText", "This Application");
+			}
+
+			if (FirstInstance->GetDeviceName() == FPlatformProcess::ComputerName())
+			{
+				return LOCTEXT("UnnamedLocalSessionText", "Unnamed Session (Local)");
+			}
+
+			return LOCTEXT("UnnamedRemoteSessionText", "Unnamed Session (Remote)");
+		}
+
+		return LOCTEXT("UnnamedSessionText", "Unnamed Session");
+	}
+
+private:
+
+	/** Callback for getting the text in the 'Device' column. */
+	FText HandleDeviceColumnText() const
+	{
+		ISessionInstanceInfoPtr InstanceInfo = Item->GetInstanceInfo();
+
+		if (InstanceInfo.IsValid())
+		{
+			return FText::FromString(InstanceInfo->GetDeviceName());
+		}
+
+		return FText::GetEmpty();
+	}
+
+	/** Callback for getting the border color for this row. */
+	FSlateColor HandleInstanceBorderBackgroundColor() const
+	{
+		ISessionInstanceInfoPtr InstanceInfo = Item->GetInstanceInfo();
+
+		if (InstanceInfo.IsValid())
+		{
+			return FLinearColor((GetTypeHash(InstanceInfo->GetInstanceId()) & 0xff) * 360.0 / 256.0, 0.8, 0.3, 1.0).HSVToLinearRGB();
+		}
+
+		return FLinearColor::Transparent;
+	}
+
+	/** Callback for getting the border brush for this row. */
+	const FSlateBrush* HandleInstanceBorderBrush() const
+	{
+		ISessionInstanceInfoPtr InstanceInfo = Item->GetInstanceInfo();
+
+		if (InstanceInfo.IsValid())
+		{
+			if (FDateTime::UtcNow() - InstanceInfo->GetLastUpdateTime() < FTimespan::FromSeconds(10.0))
+			{
+				return FEditorStyle::GetBrush("ErrorReporting.Box");
+			}
+
+			return FEditorStyle::GetBrush("ErrorReporting.EmptyBox");
+		}
+
+		return nullptr;
+	}
+
+	/** Callback for getting the type of the session instance. */
+	FText HandleInstanceTypeText() const
+	{
+		ISessionInstanceInfoPtr InstanceInfo = Item->GetInstanceInfo();
+
+		if (InstanceInfo.IsValid())
+		{
+			return FText::FromString(InstanceInfo->GetInstanceType());
+		}
+
+		return FText::GetEmpty();
+	}
+
+	/** Callback for getting the instance's current level. */
+	FText HandleLevelColumnText() const
+	{
+		ISessionInstanceInfoPtr InstanceInfo = Item->GetInstanceInfo();
+
+		if (InstanceInfo.IsValid())
+		{
+			return FText::FromString(InstanceInfo->GetCurrentLevel());
+		}
+
+		return FText::GetEmpty();
+	}
+
+	/** Callback for getting the text in the 'Status' column. */
+	FText HandleStatusColumnText() const
+	{
+		ISessionInstanceInfoPtr InstanceInfo = Item->GetInstanceInfo();
+
+		if (InstanceInfo.IsValid())
+		{
+			if (FDateTime::UtcNow() - InstanceInfo->GetLastUpdateTime() < FTimespan::FromSeconds(10.0))
+			{
+				return LOCTEXT("StatusRunning", "Running");
+			}
+
+			return LOCTEXT("StatusTimedOut", "Timed Out");
+		}
+
+		return FText::GetEmpty();
+	}
+
+	/** Callback for getting the foreground text color. */
+	FSlateColor HandleTextColorAndOpacity() const
+	{
+		ISessionInstanceInfoPtr InstanceInfo = Item->GetInstanceInfo();
+
+		if (InstanceInfo.IsValid())
+		{
+			if (FDateTime::UtcNow() - InstanceInfo->GetLastUpdateTime() < FTimespan::FromSeconds(10.0))
+			{
+				return FSlateColor::UseForeground();
+			}
+		}
+
+		return FSlateColor::UseSubduedForeground();
+	}
+
+private:
+
+	/** The highlight string for this row. */
+	TAttribute<FText> HighlightText;
+
+	/** A reference to the tree item that is displayed in this row. */
+	TSharedPtr<FSessionBrowserInstanceTreeItem> Item;
+};
+
+
+#undef LOCTEXT_NAMESPACE
diff --git a/Engine/Source/Developer/SessionFrontend/Private/Widgets/Browser/SSessionBrowserTreeSessionRow.h b/Engine/Source/Developer/SessionFrontend/Private/Widgets/Browser/SSessionBrowserTreeSessionRow.h
new file mode 100644
--- /dev/null
+++ b/Engine/Source/Developer/SessionFrontend/Private/Widgets/Browser/SSessionBrowserTreeSessionRow.h
@@ -0,0 +1,181 @@
+// Copyright 1998-2015 Epic Games, Inc. All Rights Reserved.
+
+#pragma once
+
+
+#define LOCTEXT_NAMESPACE "SSessionBrowserTreeRow"
+
+
+/**
+ * Implements a row widget for the session browser tree.
+ */
+class SSessionBrowserTreeSessionRow
+	: public STableRow<TSharedPtr<FSessionBrowserSessionTreeItem>>
+{
+public:
+
+	SLATE_BEGIN_ARGS(SSessionBrowserTreeSessionRow) { }
+		SLATE_ATTRIBUTE(FText, HighlightText)
+		SLATE_ARGUMENT(TSharedPtr<STableViewBase>, OwnerTableView)
+		SLATE_ARGUMENT(TSharedPtr<FSessionBrowserSessionTreeItem>, Item)
+	SLATE_END_ARGS()
+
+public:
+
+	/**
+	 * Constructs the widget.
+	 *
+	 * @param InArgs The construction arguments.
+	 */
+	BEGIN_SLATE_FUNCTION_BUILD_OPTIMIZATION
+	void Construct(const FArguments& InArgs, const TSharedRef<STableViewBase>& InOwnerTableView)
+	{
+		HighlightText = InArgs._HighlightText;
+		Item = InArgs._Item;
+
+		ChildSlot
+			.Padding(0.0f, 2.0f, 0.0f, 0.0f)
+			[
+				SNew(SBox)
+					.Padding(FMargin(0.0f, 3.0f, 16.0f, 3.0f))
+					.ToolTipText(this, &SSessionBrowserTreeSessionRow::HandleBorderToolTipText)
+					[
+						SNew(SHorizontalBox)
+						
+						+ SHorizontalBox::Slot()
+							.AutoWidth()
+							.VAlign(VAlign_Center)
+							.Padding(2.0f, 2.0f, 2.0f, 2.0f)
+							[
+								SNew(SExpanderArrow, SharedThis(this))
+									.IndentAmount(0.0f)
+							]
+
+						+ SHorizontalBox::Slot()
+							.FillWidth(1.0f)
+							.VAlign(VAlign_Center)
+							[
+								SNew(STextBlock)
+									.Text(SanitizeSessionName(Item->GetSessionInfo()))
+									.Font(FEditorStyle::GetFontStyle("PropertyWindow.NormalFont"))
+							]
+					]
+			];
+
+		STableRow<TSharedPtr<FSessionBrowserSessionTreeItem>>::ConstructInternal(
+			STableRow<TSharedPtr<FSessionBrowserSessionTreeItem>>::FArguments()
+				.ShowSelection(false)
+				.Style(FEditorStyle::Get(), "DetailsView.TreeView.TableRow"),
+			InOwnerTableView
+		);
+	}
+	END_SLATE_FUNCTION_BUILD_OPTIMIZATION
+
+protected:
+
+	FText SanitizeSessionName(const ISessionInfoPtr& SessionInfo) const
+	{
+		const FString& SessionName = SessionInfo->GetSessionName();
+		const FString& SessionOwner = SessionInfo->GetSessionOwner();
+
+		// return name of a launched session
+		if (!SessionInfo->IsStandalone() || !SessionName.IsEmpty())
+		{
+			if (SessionInfo->GetSessionOwner() != FPlatformProcess::UserName(true))
+			{
+				return FText::Format(LOCTEXT("SessionNameFormat", "{0} - {1}"), FText::FromString(SessionName), FText::FromString(SessionOwner));
+			}
+
+			return FText::FromString(SessionName);
+		}
+
+		// generate name for a standalone session
+		TArray<ISessionInstanceInfoPtr> Instances;
+		SessionInfo->GetInstances(Instances);
+
+		if (Instances.Num() > 0)
+		{
+			const ISessionInstanceInfoPtr& FirstInstance = Instances[0];
+
+			if ((Instances.Num() == 1) && (FirstInstance->GetInstanceId() == FApp::GetInstanceId()))
+			{
+				return LOCTEXT("ThisApplicationSessionText", "This Application");
+			}
+		}
+
+		if (SessionInfo->GetSessionOwner() != FPlatformProcess::UserName(true))
+		{
+			return FText::Format(LOCTEXT("UnnamedSessionFormat", "Unnamed - {0}"), FText::FromString(SessionOwner));
+		}
+
+		return LOCTEXT("UnnamedSessionText", "Unnamed");
+	}
+
+protected:
+
+	// SWidget overrides
+
+	virtual FReply OnMouseButtonDown(const FGeometry& MyGeometry, const FPointerEvent& MouseEvent) override
+	{
+		if (MouseEvent.GetEffectingButton() == EKeys::LeftMouseButton)
+		{
+			ToggleExpansion();
+
+			return FReply::Handled();
+		}
+		else
+		{
+			return FReply::Unhandled();
+		}
+	}
+
+private:
+
+	/** Callback for getting the background image of the row's border. */
+	const FSlateBrush* HandleBorderBackgroundImage() const
+	{
+		if (IsHovered())
+		{
+			return IsItemExpanded()
+				? FEditorStyle::GetBrush("DetailsView.CategoryTop_Hovered")
+				: FEditorStyle::GetBrush("DetailsView.CollapsedCategory_Hovered");
+		}
+		else
+		{
+			return IsItemExpanded()
+				? FEditorStyle::GetBrush("DetailsView.CategoryTop")
+				: FEditorStyle::GetBrush("DetailsView.CollapsedCategory");
+		}
+	}
+
+	/** Callback for getting the text of the row border's tool tip. */
+	FText HandleBorderToolTipText() const
+	{
+		FTextBuilder ToolTipTextBuilder;
+
+		ISessionInfoPtr SessionInfo = Item->GetSessionInfo();
+
+		if (SessionInfo.IsValid())
+		{
+			ToolTipTextBuilder.AppendLineFormat(LOCTEXT("SessionToolTipSessionId", "Session ID: {0}"), FText::FromString(SessionInfo->GetSessionId().ToString(EGuidFormats::DigitsWithHyphensInBraces)));
+			ToolTipTextBuilder.AppendLine();
+			ToolTipTextBuilder.AppendLineFormat(LOCTEXT("SessionToolTipNumInstances", "Total Instances: {0}"), FText::AsNumber(SessionInfo->GetNumInstances()));
+			ToolTipTextBuilder.AppendLineFormat(LOCTEXT("SessionToolTipIsStandalone", "Is Standalone: {0}"), SessionInfo->IsStandalone() ? GYes : GNo);
+			ToolTipTextBuilder.AppendLine();
+			ToolTipTextBuilder.AppendLineFormat(LOCTEXT("SessionToolTipLastUpdateTime", "Last Update Time: {0}"), FText::AsDateTime(SessionInfo->GetLastUpdateTime()));
+		}
+
+		return ToolTipTextBuilder.ToText();
+	}
+
+private:
+
+	/** The highlight string for this row. */
+	TAttribute<FText> HighlightText;
+
+	/** A reference to the tree item that is displayed in this row. */
+	TSharedPtr<FSessionBrowserSessionTreeItem> Item;
+};
+
+
+#undef LOCTEXT_NAMESPACE
diff --git a/Engine/Source/Developer/SessionFrontend/Private/Widgets/Console/SSessionConsole.cpp b/Engine/Source/Developer/SessionFrontend/Private/Widgets/Console/SSessionConsole.cpp
--- a/Engine/Source/Developer/SessionFrontend/Private/Widgets/Console/SSessionConsole.cpp
+++ b/Engine/Source/Developer/SessionFrontend/Private/Widgets/Console/SSessionConsole.cpp
@@ -145,7 +145,7 @@ void SSessionConsole::Construct(const FArguments& InArgs, ISessionManagerRef InS
 					.Visibility(this, &SSessionConsole::HandleSelectSessionOverlayVisibility)
 					[
 						SNew(STextBlock)
-							.Text(LOCTEXT("SelectSessionOverlayText", "Please select a session from the Session Browser"))
+							.Text(LOCTEXT("SelectSessionOverlayText", "Please select at least one instance from the Session Browser"))
 					]
 			]
 	];
@@ -220,8 +220,7 @@ void SSessionConsole::ReloadLog(bool FullyReload)
 	{
 		AvailableLogs.Reset();
 
-		TArray<ISessionInstanceInfoPtr> SelectedInstances;
-		SessionManager->GetSelectedInstances(SelectedInstances);
+		const auto& SelectedInstances = SessionManager->GetSelectedInstances();
 
 		for (const auto& Instance : SelectedInstances)
 		{
@@ -338,10 +337,7 @@ void SSessionConsole::SendCommand(const FString& CommandString)
 		return;
 	}
 
-	TArray<ISessionInstanceInfoPtr> SelectedInstances;
-	SessionManager->GetSelectedInstances(SelectedInstances);
-
-	for (auto& Instance : SelectedInstances)
+	for (auto& Instance : SessionManager->GetSelectedInstances())
 	{
 		Instance->ExecuteCommand(CommandString);
 	}
@@ -452,7 +448,7 @@ FText SSessionConsole::HandleLogListGetHighlightText() const
 
 bool SSessionConsole::HandleMainContentIsEnabled() const
 {
-	return SessionManager->GetSelectedSession().IsValid();
+	return (SessionManager->GetSelectedInstances().Num() > 0);
 }
 
 
@@ -470,7 +466,7 @@ bool SSessionConsole::HandleSaveActionCanExecute()
 
 EVisibility SSessionConsole::HandleSelectSessionOverlayVisibility() const
 {
-	if (SessionManager->GetSelectedSession().IsValid())
+	if (SessionManager->GetSelectedInstances().Num() > 0)
 	{
 		return EVisibility::Hidden;
 	}
@@ -479,7 +475,7 @@ EVisibility SSessionConsole::HandleSelectSessionOverlayVisibility() const
 }
 
 
-void SSessionConsole::HandleSessionManagerInstanceSelectionChanged()
+void SSessionConsole::HandleSessionManagerInstanceSelectionChanged(const TSharedPtr<ISessionInstanceInfo>& /*Instance*/, bool /*Selected*/)
 {
 	ReloadLog(true);
 }
diff --git a/Engine/Source/Developer/SessionFrontend/Private/Widgets/Console/SSessionConsole.h b/Engine/Source/Developer/SessionFrontend/Private/Widgets/Console/SSessionConsole.h
--- a/Engine/Source/Developer/SessionFrontend/Private/Widgets/Console/SSessionConsole.h
+++ b/Engine/Source/Developer/SessionFrontend/Private/Widgets/Console/SSessionConsole.h
@@ -127,7 +127,7 @@ private:
 	EVisibility HandleSelectSessionOverlayVisibility() const;
 
 	/** Callback for changing the engine instance selection. */
-	void HandleSessionManagerInstanceSelectionChanged();
+	void HandleSessionManagerInstanceSelectionChanged(const TSharedPtr<ISessionInstanceInfo>& Instance, bool Selected);
 
 	/** Callback for received log entries. */
 	void HandleSessionManagerLogReceived(const ISessionInfoRef& Session, const ISessionInstanceInfoRef& Instance, const FSessionLogMessageRef& Message);
diff --git a/Engine/Source/Developer/SessionFrontend/Private/Widgets/SSessionFrontend.cpp b/Engine/Source/Developer/SessionFrontend/Private/Widgets/SSessionFrontend.cpp
--- a/Engine/Source/Developer/SessionFrontend/Private/Widgets/SSessionFrontend.cpp
+++ b/Engine/Source/Developer/SessionFrontend/Private/Widgets/SSessionFrontend.cpp
@@ -55,18 +55,18 @@ void SSessionFrontend::Construct( const FArguments& InArgs, const TSharedRef<SDo
 		.SetGroup(AppMenuGroup);
 	
 	// create tab layout
-	const TSharedRef<FTabManager::FLayout> Layout = FTabManager::NewLayout("SessionFrontendLayout_v1.1")
+	const TSharedRef<FTabManager::FLayout> Layout = FTabManager::NewLayout("SessionFrontendLayout_v1.2")
 		->AddArea
 		(
 			FTabManager::NewPrimaryArea()
-				->SetOrientation(Orient_Vertical)
+				->SetOrientation(Orient_Horizontal)
 				->Split
 				(
 					// session browser
 					FTabManager::NewStack()
 						->AddTab(SessionBrowserTabId, ETabState::OpenedTab)
 						->SetHideTabWell(true)
-						->SetSizeCoefficient(0.15f)
+						->SetSizeCoefficient(0.25f)
 				)
 				->Split
 				(
@@ -76,7 +76,7 @@ void SSessionFrontend::Construct( const FArguments& InArgs, const TSharedRef<SDo
 						->AddTab(AutomationTabId, ETabState::OpenedTab)
 						->AddTab(SessionScreenTabId, ETabState::ClosedTab)
 						->AddTab(ProfilerTabId, ETabState::OpenedTab)
-						->SetSizeCoefficient(0.85f)
+						->SetSizeCoefficient(0.75f)
 						->SetForegroundTab(SessionConsoleTabId)
 				)							
 		);
diff --git a/Engine/Source/Developer/SessionServices/Private/SessionManager.cpp b/Engine/Source/Developer/SessionServices/Private/SessionManager.cpp
--- a/Engine/Source/Developer/SessionServices/Private/SessionManager.cpp
+++ b/Engine/Source/Developer/SessionServices/Private/SessionManager.cpp
@@ -6,7 +6,7 @@
 /* FSessionManager structors
  *****************************************************************************/
 
-FSessionManager::FSessionManager( const IMessageBusRef& InMessageBus )
+FSessionManager::FSessionManager(const IMessageBusRef& InMessageBus)
 	: MessageBusPtr(InMessageBus)
 {
 	// fill in the owner array
@@ -49,26 +49,15 @@ FSessionManager::~FSessionManager()
 /* ISessionManager interface
  *****************************************************************************/
 
-void FSessionManager::AddOwner( const FString& InOwner )
+void FSessionManager::AddOwner(const FString& InOwner)
 {
 	FilteredOwners.Add(InOwner);
 }
 
 
-void FSessionManager::GetSelectedInstances( TArray<ISessionInstanceInfoPtr>& OutInstances) const
+const TArray<ISessionInstanceInfoPtr>& FSessionManager::GetSelectedInstances() const
 {
-	if (SelectedSession.IsValid())
-	{
-		SelectedSession->GetInstances(OutInstances);
-
-		for (int32 Index = OutInstances.Num() - 1; Index >= 0 ; --Index)
-		{
-			if (DeselectedInstances.Contains(OutInstances[Index]))
-			{
-				OutInstances.RemoveAtSwap(Index);
-			}
-		}
-	}
+	return SelectedInstances;
 }
 
 
@@ -78,7 +67,7 @@ const ISessionInfoPtr& FSessionManager::GetSelectedSession() const
 }
 
 
-void FSessionManager::GetSessions( TArray<ISessionInfoPtr>& OutSessions ) const
+void FSessionManager::GetSessions(TArray<ISessionInfoPtr>& OutSessions) const
 {
 	OutSessions.Empty(Sessions.Num());
 
@@ -89,9 +78,9 @@ void FSessionManager::GetSessions( TArray<ISessionInfoPtr>& OutSessions ) const
 }
 
 
-bool FSessionManager::IsInstanceSelected( const TSharedRef<ISessionInstanceInfo>& Instance ) const
+bool FSessionManager::IsInstanceSelected(const TSharedRef<ISessionInstanceInfo>& Instance) const
 {
-	return ((Instance->GetOwnerSession() == SelectedSession) && !DeselectedInstances.Contains(Instance));
+	return ((Instance->GetOwnerSession() == SelectedSession) && SelectedInstances.Contains(Instance));
 }
 
 
@@ -107,71 +96,76 @@ FSimpleMulticastDelegate& FSessionManager::OnSessionInstanceUpdated()
 }
 
 
-void FSessionManager::RemoveOwner( const FString& InOwner )
+void FSessionManager::RemoveOwner(const FString& InOwner)
 {
 	FilteredOwners.Remove(InOwner);
 	RefreshSessions();
 }
 
 
-bool FSessionManager::SelectSession( const ISessionInfoPtr& Session )
+bool FSessionManager::SelectSession(const ISessionInfoPtr& Session)
 {
-	if (Session != SelectedSession)
+	// already selected?
+	if (Session == SelectedSession)
 	{
-		if (!Session.IsValid() || Sessions.Contains(Session->GetSessionId()))
-		{
-			bool CanSelect = true;
+		return true;
+	}
 
-			CanSelectSessionDelegate.Broadcast(Session, CanSelect);
+	// do we own the session?
+	if (Session.IsValid() && !Sessions.Contains(Session->GetSessionId()))
+	{
+		return false;
+	}
 
-			if (CanSelect)
-			{
-				SelectedSession = Session;
+	// allowed to de/select?
+	bool CanSelect = true;
+	CanSelectSessionDelegate.Broadcast(Session, CanSelect);
 
-				SelectedSessionChangedEvent.Broadcast(Session);
-			}
-			else
-			{
-				return false;
-			}
-		}
+	if (!CanSelect)
+	{
+		return false;
 	}
 
+	// set selection
+	SelectedInstances.Empty();
+	SelectedSession = Session;
+	SelectedSessionChangedEvent.Broadcast(Session);
+
 	return true;
 }
 
 
-bool FSessionManager::SetInstanceSelected( const ISessionInstanceInfoPtr& Instance, bool Selected )
+bool FSessionManager::SetInstanceSelected(const ISessionInstanceInfoRef& Instance, bool Selected)
 {
-	if (Instance->GetOwnerSession() == SelectedSession)
+	if (Instance->GetOwnerSession() != SelectedSession)
+	{
+		return false;
+	}
+
+	if (Selected)
 	{
-		if (Selected)
+		if (!SelectedInstances.Contains(Instance))
 		{
-			if (DeselectedInstances.Remove(Instance) > 0)
-			{
-				InstanceSelectionChangedDelegate.Broadcast();
-			}
+			SelectedInstances.Add(Instance);
+			InstanceSelectionChangedDelegate.Broadcast(Instance, true);
 		}
-		else
+	}
+	else
+	{
+		if (SelectedInstances.Remove(Instance) > 0)
 		{
-			if (!DeselectedInstances.Contains(Instance))
-			{
-				DeselectedInstances.Add(Instance);
-				InstanceSelectionChangedDelegate.Broadcast();
-			}
+			InstanceSelectionChangedDelegate.Broadcast(Instance, false);
 		}
-
-		return true;
 	}
 
-	return false;
+	return true;
 }
 
 
 /* FSessionManager implementation
  *****************************************************************************/
 
-void FSessionManager::FindExpiredSessions( const FDateTime& Now )
+void FSessionManager::FindExpiredSessions(const FDateTime& Now)
 {
 	bool Dirty = false;
 
@@ -191,20 +185,18 @@ void FSessionManager::FindExpiredSessions( const FDateTime& Now )
 }
 
 
-bool FSessionManager::IsValidOwner( const FString& Owner )
+bool FSessionManager::IsValidOwner(const FString& Owner)
 {
 	if (Owner == FPlatformProcess::UserName(true))
 	{
 		return true;
 	}
-	else 
+
+	for (const auto& FilteredOwner : FilteredOwners)
 	{
-		for (int32 Index = 0; Index < FilteredOwners.Num(); ++Index)
+		if (FilteredOwner == Owner)
 		{
-			if (FilteredOwners[Index] == Owner)
-			{
-				return true;
-			}
+			return true;
 		}
 	}
 
@@ -221,7 +213,6 @@ void FSessionManager::RefreshSessions()
 		if (!IsValidOwner(It.Value()->GetSessionOwner()))
 		{
 			It.RemoveCurrent();
-
 			Dirty = true;
 		}
 	}
@@ -248,7 +239,7 @@ void FSessionManager::SendPing()
 /* FSessionManager callbacks
  *****************************************************************************/
 
-void FSessionManager::HandleEnginePongMessage( const FEngineServicePong& Message, const IMessageContextRef& Context )
+void FSessionManager::HandleEnginePongMessage(const FEngineServicePong& Message, const IMessageContextRef& Context)
 {
 	if (!Message.SessionId.IsValid())
 	{
@@ -266,16 +257,7 @@ void FSessionManager::HandleEnginePongMessage( const FEngineServicePong& Message
 }
 
 
-void FSessionManager::HandleInstanceDiscovered( const ISessionInfoRef& Session, const ISessionInstanceInfoRef& Instance )
-{
-	if (Session == SelectedSession)
-	{
-		InstanceSelectionChangedDelegate.Broadcast();
-	}	
-}
-
-
-void FSessionManager::HandleLogReceived( const ISessionInfoRef& Session, const ISessionInstanceInfoRef& Instance, const FSessionLogMessageRef& Message )
+void FSessionManager::HandleLogReceived(const ISessionInfoRef& Session, const ISessionInstanceInfoRef& Instance, const FSessionLogMessageRef& Message)
 {
 	if (Session == SelectedSession)
 	{
@@ -284,7 +266,7 @@ void FSessionManager::HandleLogReceived( const ISessionInfoRef& Session, const I
 }
 
 
-void FSessionManager::HandleSessionPongMessage( const FSessionServicePong& Message, const IMessageContextRef& Context )
+void FSessionManager::HandleSessionPongMessage(const FSessionServicePong& Message, const IMessageContextRef& Context)
 {
 	if (!Message.SessionId.IsValid())
 	{
@@ -313,7 +295,6 @@ void FSessionManager::HandleSessionPongMessage( const FSessionServicePong& Messa
 	else
 	{
 		Session = MakeShareable(new FSessionInfo(Message.SessionId, MessageBus.ToSharedRef()));
-		Session->OnInstanceDiscovered().AddSP(this, &FSessionManager::HandleInstanceDiscovered);
 		Session->OnLogReceived().AddSP(this, &FSessionManager::HandleLogReceived);
 		Session->UpdateFromMessage(Message, Context);
 
@@ -322,7 +303,7 @@ void FSessionManager::HandleSessionPongMessage( const FSessionServicePong& Messa
 }
 
 
-bool FSessionManager::HandleTicker( float DeltaTime )
+bool FSessionManager::HandleTicker(float DeltaTime)
 {
 	FDateTime Now = FDateTime::UtcNow();
 
diff --git a/Engine/Source/Developer/SessionServices/Private/SessionManager.h b/Engine/Source/Developer/SessionServices/Private/SessionManager.h
--- a/Engine/Source/Developer/SessionServices/Private/SessionManager.h
+++ b/Engine/Source/Developer/SessionServices/Private/SessionManager.h
@@ -19,7 +19,7 @@ public:
 	 *
 	 * @param InMessageBus The message bus to use.
 	 */
-	FSessionManager( const IMessageBusRef& InMessageBus );
+	FSessionManager(const IMessageBusRef& InMessageBus);
 	
 	/** Destructor */
 	~FSessionManager();
@@ -28,11 +28,11 @@ public:
 
 	// ISessionManager interface
 
-	virtual void AddOwner( const FString& InOwner ) override;
-	virtual void GetSelectedInstances( TArray<TSharedPtr<ISessionInstanceInfo>>& OutInstances) const override;
+	virtual void AddOwner(const FString& InOwner) override;
+	virtual const TArray<ISessionInstanceInfoPtr>& GetSelectedInstances() const override;
 	virtual const ISessionInfoPtr& GetSelectedSession() const override;
-	virtual void GetSessions( TArray<TSharedPtr<ISessionInfo>>& OutSessions ) const override;
-	virtual bool IsInstanceSelected( const TSharedRef<ISessionInstanceInfo>& Instance ) const override;
+	virtual void GetSessions(TArray<TSharedPtr<ISessionInfo>>& OutSessions) const override;
+	virtual bool IsInstanceSelected(const TSharedRef<ISessionInstanceInfo>& Instance) const override;
 
 	DECLARE_DERIVED_EVENT(FSessionManager, ISessionManager::FCanSelectSessionEvent, FCanSelectSessionEvent)
 	virtual FCanSelectSessionEvent& OnCanSelectSession() override
@@ -60,9 +60,9 @@ public:
 
 	virtual FSimpleMulticastDelegate& OnSessionsUpdated() override;
 	virtual FSimpleMulticastDelegate& OnSessionInstanceUpdated() override;
-	virtual void RemoveOwner( const FString& InOwner ) override;
-	virtual bool SelectSession( const ISessionInfoPtr& Session ) override;
-	virtual bool SetInstanceSelected( const ISessionInstanceInfoPtr& Instance, bool Selected ) override;
+	virtual void RemoveOwner(const FString& InOwner) override;
+	virtual bool SelectSession(const ISessionInfoPtr& Session) override;
+	virtual bool SetInstanceSelected(const ISessionInstanceInfoRef& Instance, bool Selected) override;
 
 protected:
 
@@ -71,14 +71,14 @@ protected:
 	 *
 	 * @param Now The current time.
 	 */
-	void FindExpiredSessions( const FDateTime& Now );
+	void FindExpiredSessions(const FDateTime& Now);
 
 	/**
 	 * Checks whether the specified owner is valid.
 	 *
 	 * @return true if the owner is valid, false otherwise.
 	 */
-	bool IsValidOwner( const FString& Owner );
+	bool IsValidOwner(const FString& Owner);
 
 	/** Refresh the sessions based on the owner filter list. */
 	void RefreshSessions();
@@ -89,28 +89,22 @@ protected:
 private:
 
 	/** Callback for handling FSessionServicePong messages. */
-	void HandleEnginePongMessage( const FEngineServicePong& Message, const IMessageContextRef& Context );
-
-	/** Callback for newly discovered instances. */
-	void HandleInstanceDiscovered( const ISessionInfoRef& Session, const ISessionInstanceInfoRef& Instance );
+	void HandleEnginePongMessage(const FEngineServicePong& Message, const IMessageContextRef& Context);
 
 	/** Callback received log entries. */
-	void HandleLogReceived( const ISessionInfoRef& Session, const ISessionInstanceInfoRef& Instance, const FSessionLogMessageRef& Message );
+	void HandleLogReceived(const ISessionInfoRef& Session, const ISessionInstanceInfoRef& Instance, const FSessionLogMessageRef& Message);
 
 	/** Callback for handling FSessionServicePong messages. */
-	void HandleSessionPongMessage( const FSessionServicePong& Message, const IMessageContextRef& Context );
+	void HandleSessionPongMessage(const FSessionServicePong& Message, const IMessageContextRef& Context);
 
 	/** Callback for ticks from the ticker. */
-	bool HandleTicker( float DeltaTime );
+	bool HandleTicker(float DeltaTime);
 
 private:
 
 	/** The address of the automation controller to where we can forward any automation workers found. */
 	FGuid AutomationControllerAddress;
 
-	/** Holds the list of currently selected instances. */
-	TArray<ISessionInstanceInfoPtr> DeselectedInstances;
-
 	/** Holds the time at which the last ping was sent. */
 	FDateTime LastPingTime;
 
@@ -120,6 +114,9 @@ private:
 	/** Holds the messaging endpoint. */
 	FMessageEndpointPtr MessageEndpoint;
 
+	/** Holds the list of currently selected instances. */
+	TArray<ISessionInstanceInfoPtr> SelectedInstances;
+
 	/** Holds a reference to the currently selected session. */
 	ISessionInfoPtr SelectedSession;
 
diff --git a/Engine/Source/Developer/SessionServices/Public/ISessionManager.h b/Engine/Source/Developer/SessionServices/Public/ISessionManager.h
--- a/Engine/Source/Developer/SessionServices/Public/ISessionManager.h
+++ b/Engine/Source/Developer/SessionServices/Public/ISessionManager.h
@@ -25,9 +25,9 @@ public:
 	/**
 	  * Gets the collection of currently selected engine instances.
 	  *
-	  * @param OutInstances Will hold the selected instances.
+	  * @return The selected instances.
 	  */
-	virtual void GetSelectedInstances( TArray<TSharedPtr<ISessionInstanceInfo>>& OutInstances) const = 0;
+	virtual const TArray<TSharedPtr<ISessionInstanceInfo>>& GetSelectedInstances() const = 0;
 
 	/**
 	 * Get the selected session - as chosen in the session browser
@@ -72,7 +72,7 @@ public:
 	 * @param Selected Whether the instance should be selected (true) or unselected (false).
 	 * @return true if the instance was selected, false otherwise.
 	 */
-	virtual bool SetInstanceSelected( const TSharedPtr<ISessionInstanceInfo>& Instance, bool Selected ) = 0;
+	virtual bool SetInstanceSelected( const TSharedRef<ISessionInstanceInfo>& Instance, bool Selected ) = 0;
 
 public:
 
@@ -89,7 +89,7 @@ public:
 	 *
 	 * @return The delegate.
 	 */
-	DECLARE_EVENT(ISessionManager, FInstanceSelectionChangedEvent)
+	DECLARE_EVENT_TwoParams(ISessionManager, FInstanceSelectionChangedEvent, const TSharedPtr<ISessionInstanceInfo>& /*Instance*/, bool /*Selected*/)
 	virtual FInstanceSelectionChangedEvent& OnInstanceSelectionChanged() = 0;
 
 	/**
