Added missing virtual destructors.

diff --git a/Engine/Source/Developer/SessionFrontend/Private/Models/SessionBrowserTreeItems.h b/Engine/Source/Developer/SessionFrontend/Private/Models/SessionBrowserTreeItems.h
--- a/Engine/Source/Developer/SessionFrontend/Private/Models/SessionBrowserTreeItems.h
+++ b/Engine/Source/Developer/SessionFrontend/Private/Models/SessionBrowserTreeItems.h
@@ -16,6 +16,11 @@ enum class ESessionBrowserTreeNodeType
  */
 class FSessionBrowserTreeItem
 {
+public:
+
+	/** Virtual destructor. */
+	virtual ~FSessionBrowserTreeItem() { }
+
 public:
 
 	/**
@@ -100,6 +105,9 @@ public:
 		, ToolTipText(InToolTipText)
 	{ }
 
+	/** Virtual destructor. */
+	virtual ~FSessionBrowserGroupTreeItem() { }
+
 public:
 
 	/**
@@ -154,6 +162,9 @@ public:
 		: InstanceInfo(InInstanceInfo)
 	{ }
 
+	/** Virtual destructor. */
+	virtual ~FSessionBrowserInstanceTreeItem() { }
+
 public:
 
 	/**
@@ -195,6 +206,9 @@ public:
 		: SessionInfo(InSessionInfo)
 	{ }
 
+	/** Virtual destructor. */
+	virtual ~FSessionBrowserSessionTreeItem() { }
+
 public:
 
 	/**
