SessionBrowser: Removed dead code

diff --git a/Engine/Source/Developer/SessionFrontend/Private/Widgets/Browser/SSessionBrowserTreeInstanceRow.h b/Engine/Source/Developer/SessionFrontend/Private/Widgets/Browser/SSessionBrowserTreeInstanceRow.h
--- a/Engine/Source/Developer/SessionFrontend/Private/Widgets/Browser/SSessionBrowserTreeInstanceRow.h
+++ b/Engine/Source/Developer/SessionFrontend/Private/Widgets/Browser/SSessionBrowserTreeInstanceRow.h
@@ -148,42 +148,6 @@ public:
 	}
 	END_SLATE_FUNCTION_BUILD_OPTIMIZATION
 
-protected:
-
-	FText SanitizeSessionName(const ISessionInfoPtr& SessionInfo) const
-	{
-		const FString& SessionName = SessionInfo->GetSessionName();
-
-		// return name of a launched session
-		if (!SessionInfo->IsStandalone() || !SessionName.IsEmpty())
-		{
-			return FText::FromString(SessionName);
-		}
-
-		// generate name for a standalone session
-		TArray<ISessionInstanceInfoPtr> Instances;
-		SessionInfo->GetInstances(Instances);
-
-		if (Instances.Num() > 0)
-		{
-			const ISessionInstanceInfoPtr& FirstInstance = Instances[0];
-
-			if ((Instances.Num() == 1) && (FirstInstance->GetInstanceId() == FApp::GetInstanceId()))
-			{
-				return LOCTEXT("ThisApplicationSessionText", "This Application");
-			}
-
-			if (FirstInstance->GetDeviceName() == FPlatformProcess::ComputerName())
-			{
-				return LOCTEXT("UnnamedLocalSessionText", "Unnamed Session (Local)");
-			}
-
-			return LOCTEXT("UnnamedRemoteSessionText", "Unnamed Session (Remote)");
-		}
-
-		return LOCTEXT("UnnamedSessionText", "Unnamed Session");
-	}
-
 private:
 
 	/** Callback for getting the text in the 'Device' column. */
