SessionBrowser: Fixed unnamed non-standalone sessions not showing as 'unnamed'

diff --git a/Engine/Source/Developer/SessionFrontend/Private/Widgets/Browser/SSessionBrowserTreeSessionRow.h b/Engine/Source/Developer/SessionFrontend/Private/Widgets/Browser/SSessionBrowserTreeSessionRow.h
--- a/Engine/Source/Developer/SessionFrontend/Private/Widgets/Browser/SSessionBrowserTreeSessionRow.h
+++ b/Engine/Source/Developer/SessionFrontend/Private/Widgets/Browser/SSessionBrowserTreeSessionRow.h
@@ -79,7 +79,7 @@ protected:
 		const FString& SessionOwner = SessionInfo->GetSessionOwner();
 
 		// return name of a launched session
-		if (!SessionInfo->IsStandalone() || !SessionName.IsEmpty())
+		if (!SessionName.IsEmpty())
 		{
 			if (SessionInfo->GetSessionOwner() != FPlatformProcess::UserName(true))
 			{
