SessionFrontend: Highlighting session groups and their number of sessions

diff --git a/Engine/Source/Developer/SessionFrontend/Private/Widgets/Browser/SSessionBrowserTreeGroupRow.h b/Engine/Source/Developer/SessionFrontend/Private/Widgets/Browser/SSessionBrowserTreeGroupRow.h
--- a/Engine/Source/Developer/SessionFrontend/Private/Widgets/Browser/SSessionBrowserTreeGroupRow.h
+++ b/Engine/Source/Developer/SessionFrontend/Private/Widgets/Browser/SSessionBrowserTreeGroupRow.h
@@ -49,7 +49,8 @@ public:
 							.VAlign(VAlign_Center)
 							[
 								SNew(STextBlock)
-									.Text(Item->GetGroupName())
+									.ColorAndOpacity(this, &SSessionBrowserTreeGroupRow::HandleGroupNameColorAndOpacity)
+									.Text(this, &SSessionBrowserTreeGroupRow::HandleGroupNameText)
 									.Font(FEditorStyle::GetFontStyle("DetailsView.CategoryFontStyle"))
 									.ShadowOffset(FVector2D(1.0f, 1.0f))
 							]
@@ -110,6 +111,27 @@ private:
 		}
 	}
 
+	/** Callback for getting the group name text's color. */
+	FSlateColor HandleGroupNameColorAndOpacity() const
+	{
+		return (Item->GetChildren().Num() == 0)
+			? FSlateColor::UseSubduedForeground()
+			: FSlateColor::UseForeground();
+	}
+
+	/** Callback for getting the group name text. */
+	FText HandleGroupNameText() const
+	{
+		int32 NumSessions = Item->GetChildren().Num();
+
+		if ((NumSessions > 0) && (Item->GetChildren()[0]->GetType() == ESessionBrowserTreeNodeType::Session))
+		{
+			return FText::Format(LOCTEXT("GroupNameFormat", "{0} ({1})"), Item->GetGroupName(), FText::AsNumber(NumSessions));
+		}
+
+		return Item->GetGroupName();
+	}
+
 	/** Gets the image for the pull-down icon. */
 	const FSlateBrush* HandlePullDownImage() const
 	{
