SessionFrontend: Fixed Crash when right clicking UFE session browser section headers (UE-21238)

diff --git a/Engine/Source/Developer/SessionFrontend/Private/Widgets/Browser/SSessionBrowser.cpp b/Engine/Source/Developer/SessionFrontend/Private/Widgets/Browser/SSessionBrowser.cpp
--- a/Engine/Source/Developer/SessionFrontend/Private/Widgets/Browser/SSessionBrowser.cpp
+++ b/Engine/Source/Developer/SessionFrontend/Private/Widgets/Browser/SSessionBrowser.cpp
@@ -410,17 +410,20 @@ void SSessionBrowser::HandleSessionTreeViewSelectionChanged(const TSharedPtr<FSe
 	{
 		if (Item.IsValid())
 		{
-			const auto& InstanceInfo = StaticCastSharedPtr<FSessionBrowserInstanceTreeItem>(Item)->GetInstanceInfo();
-
-			if (InstanceInfo.IsValid())
+			if (Item->GetType() == ESessionBrowserTreeNodeType::Instance)
 			{
-				// special handling for local application
-				if (Item->GetParent() == AppGroupItem)
+				const auto& InstanceInfo = StaticCastSharedPtr<FSessionBrowserInstanceTreeItem>(Item)->GetInstanceInfo();
+
+				if (InstanceInfo.IsValid())
 				{
-					SessionManager->SelectSession(InstanceInfo->GetOwnerSession());
-				}
+					// special handling for local application
+					if (Item->GetParent() == AppGroupItem)
+					{
+						SessionManager->SelectSession(InstanceInfo->GetOwnerSession());
+					}
 
-				SessionManager->SetInstanceSelected(InstanceInfo.ToSharedRef(), true);
+					SessionManager->SetInstanceSelected(InstanceInfo.ToSharedRef(), true);
+				}
 			}
 		}
 		else
