SessionFrontend: Displaying session authorization status in session browser; placeholder icons.

diff --git a/Engine/Source/Developer/SessionFrontend/Private/Widgets/Browser/SSessionBrowserTreeInstanceRow.h b/Engine/Source/Developer/SessionFrontend/Private/Widgets/Browser/SSessionBrowserTreeInstanceRow.h
--- a/Engine/Source/Developer/SessionFrontend/Private/Widgets/Browser/SSessionBrowserTreeInstanceRow.h
+++ b/Engine/Source/Developer/SessionFrontend/Private/Widgets/Browser/SSessionBrowserTreeInstanceRow.h
@@ -118,13 +118,25 @@ public:
 		}
 		else if (ColumnName == "Status")
 		{
-			return SNew(SBox)
-				.Padding(FMargin(4.0f, 0.0f))
+			return SNew(SHorizontalBox)
+
+			+ SHorizontalBox::Slot()
+				.FillWidth(1.0f)
+				.Padding(FMargin(4.0f, 0.0f, 0.0f, 0.0f))
+				.HAlign(HAlign_Right)
 				.VAlign(VAlign_Center)
 				[
-					SNew(STextBlock)
-						.ColorAndOpacity(this, &SSessionBrowserTreeInstanceRow::HandleTextColorAndOpacity)
-						.Text(this, &SSessionBrowserTreeInstanceRow::HandleStatusColumnText)
+					SNew(SImage)
+						.Image(this, &SSessionBrowserTreeInstanceRow::HandleAuthorizedImage)
+				]
+
+			+ SHorizontalBox::Slot()
+				.AutoWidth()
+				.Padding(FMargin(2.0f, 0.0f, 0.0f, 4.0f))
+				.VAlign(VAlign_Center)
+				[
+					SNew(SImage)
+						.Image(this, &SSessionBrowserTreeInstanceRow::HandleStatusImage)
 				];
 		}
 		else if (ColumnName == "Type")
@@ -150,6 +162,19 @@ public:
 
 private:
 
+	/** Callback for getting the image of the Authorized icon. */
+	const FSlateBrush* HandleAuthorizedImage() const
+	{
+		ISessionInstanceInfoPtr InstanceInfo = Item->GetInstanceInfo();
+
+		if (InstanceInfo.IsValid() && !InstanceInfo->IsAuthorized())
+		{
+			return FEditorStyle::GetBrush("SessionBrowser.SessionLocked");
+		}
+
+		return nullptr;
+	}
+
 	/** Callback for getting the text in the 'Device' column. */
 	FText HandleDeviceColumnText() const
 	{
@@ -220,8 +245,8 @@ private:
 		return FText::GetEmpty();
 	}
 
-	/** Callback for getting the text in the 'Status' column. */
-	FText HandleStatusColumnText() const
+	/** Callback for getting the image of the Status icon. */
+	const FSlateBrush* HandleStatusImage() const
 	{
 		ISessionInstanceInfoPtr InstanceInfo = Item->GetInstanceInfo();
 
@@ -229,13 +254,13 @@ private:
 		{
 			if (FDateTime::UtcNow() - InstanceInfo->GetLastUpdateTime() < FTimespan::FromSeconds(10.0))
 			{
-				return LOCTEXT("StatusRunning", "Running");
+				return FEditorStyle::GetBrush("SessionBrowser.StatusRunning");
 			}
 
-			return LOCTEXT("StatusTimedOut", "Timed Out");
+			return FEditorStyle::GetBrush("SessionBrowser.StatusTimedOut");
 		}
 
-		return FText::GetEmpty();
+		return nullptr;
 	}
 
 	/** Callback for getting the foreground text color. */
diff --git a/Engine/Source/Editor/EditorStyle/Private/SlateEditorStyle.cpp b/Engine/Source/Editor/EditorStyle/Private/SlateEditorStyle.cpp
--- a/Engine/Source/Editor/EditorStyle/Private/SlateEditorStyle.cpp
+++ b/Engine/Source/Editor/EditorStyle/Private/SlateEditorStyle.cpp
@@ -1079,6 +1079,9 @@ void FSlateEditorStyle::FStyle::SetupGeneralStyles()
 
 	// Session Browser tab
 	{
+		Set("SessionBrowser.SessionLocked", new IMAGE_BRUSH( "Icons/icon_levels_Locked_hi_16px", Icon16x16 ) );
+		Set("SessionBrowser.StatusRunning", new IMAGE_BRUSH( "Icons/icon_status_green_16x", Icon16x16 ) );
+		Set("SessionBrowser.StatusTimedOut", new IMAGE_BRUSH( "Icons/icon_status_grey_16x", Icon16x16 ) );
 		Set("SessionBrowser.Terminate", new IMAGE_BRUSH( "Icons/icon_DevicePowerOff_40x", Icon20x20 ) );
 
 		Set("SessionBrowser.Terminate.Font", FTextBlockStyle(NormalText)
