Fixed crash that occured when deselecting session in session frontend. Caused by 2712712

diff --git a/Engine/Source/Developer/SessionFrontend/Private/Widgets/Browser/SSessionBrowser.cpp b/Engine/Source/Developer/SessionFrontend/Private/Widgets/Browser/SSessionBrowser.cpp
--- a/Engine/Source/Developer/SessionFrontend/Private/Widgets/Browser/SSessionBrowser.cpp
+++ b/Engine/Source/Developer/SessionFrontend/Private/Widgets/Browser/SSessionBrowser.cpp
@@ -429,15 +429,21 @@ void SSessionBrowser::HandleSessionTreeViewSelectionChanged(const TSharedPtr<FSe
 		else
 		{
 			// an instance got deselected
+			TArray<ISessionInstanceInfoPtr> UnselectedSessions;
 			for (const auto& InstanceInfo : SessionManager->GetSelectedInstances())
 			{
-				const auto& InstanceItem = ItemMap.FindRef(InstanceInfo->GetInstanceId());
+				const TSharedPtr<FSessionBrowserTreeItem>& InstanceItem = ItemMap.FindRef(InstanceInfo->GetInstanceId());
 
 				if (!SessionTreeView->IsItemSelected(InstanceItem))
 				{
-					SessionManager->SetInstanceSelected(InstanceInfo.ToSharedRef(), false);
+					UnselectedSessions.Add(InstanceInfo);
 				}
 			}
+
+			for (const auto& Session : UnselectedSessions)
+			{
+				SessionManager->SetInstanceSelected(Session.ToSharedRef(), false);
+			}
 		}
 	}
 	IgnoreSessionManagerEvents = false;
