Fix SessionFrontend not initializing with selected session

diff --git a/Engine/Source/Developer/SessionFrontend/Private/Widgets/Browser/SSessionBrowser.cpp b/Engine/Source/Developer/SessionFrontend/Private/Widgets/Browser/SSessionBrowser.cpp
--- a/Engine/Source/Developer/SessionFrontend/Private/Widgets/Browser/SSessionBrowser.cpp
+++ b/Engine/Source/Developer/SessionFrontend/Private/Widgets/Browser/SSessionBrowser.cpp
@@ -105,7 +105,9 @@ void SSessionBrowser::Construct( const FArguments& InArgs, TSharedRef<ISessionMa
 
 	ReloadSessions();
 
+	updatingTreeExpansion = true;
 	SessionTreeView->SetSingleExpandedItem(AppGroupItem);
+	updatingTreeExpansion = false;
 }
 END_SLATE_FUNCTION_BUILD_OPTIMIZATION
 
