Removing a different unnecessary IsConsoleBuild from the session service (it also sends the platform which is more useful in all cases)

diff --git a/Engine/Source/Developer/SessionFrontend/Private/Widgets/Browser/SSessionBrowser.cpp b/Engine/Source/Developer/SessionFrontend/Private/Widgets/Browser/SSessionBrowser.cpp
--- a/Engine/Source/Developer/SessionFrontend/Private/Widgets/Browser/SSessionBrowser.cpp
+++ b/Engine/Source/Developer/SessionFrontend/Private/Widgets/Browser/SSessionBrowser.cpp
@@ -329,7 +329,6 @@ FText SSessionBrowser::HandleSessionTreeRowGetToolTipText(TSharedPtr<FSessionBro
 			ToolTipTextBuilder.AppendLineFormat(LOCTEXT("InstanceToolTipInstanceId", "Instance ID: {0}"), FText::FromString(InstanceInfo->GetInstanceId().ToString(EGuidFormats::DigitsWithHyphensInBraces)));
 			ToolTipTextBuilder.AppendLine();
 			ToolTipTextBuilder.AppendLineFormat(LOCTEXT("InstanceToolTipBuildDate", "Build Date: {0}"), FText::FromString(InstanceInfo->GetBuildDate()));
-			ToolTipTextBuilder.AppendLineFormat(LOCTEXT("InstanceToolTipConsoleBuild", "Console Build: {0}"), InstanceInfo->IsConsole() ? LOCTEXT("LabelYes", "Yes") : LOCTEXT("LabelNo", "No"));
 			ToolTipTextBuilder.AppendLineFormat(LOCTEXT("InstanceToolTipEngineVersion", "Engine Version: {0}"), InstanceInfo->GetEngineVersion() == 0 ? LOCTEXT("CustomBuildVersion", "Custom Build") : FText::FromString(FString::FromInt(InstanceInfo->GetEngineVersion())));
 			ToolTipTextBuilder.AppendLineFormat(LOCTEXT("InstanceToolTipPlatform", "Platform: {0}"), FText::FromString(InstanceInfo->GetPlatformName()));
 			ToolTipTextBuilder.AppendLineFormat(LOCTEXT("InstanceToolTipCurrentLevel", "Current Level: {0}"), FText::FromString(InstanceInfo->GetCurrentLevel()));
diff --git a/Engine/Source/Runtime/SessionMessages/Public/SessionServiceMessages.h b/Engine/Source/Runtime/SessionMessages/Public/SessionServiceMessages.h
--- a/Engine/Source/Runtime/SessionMessages/Public/SessionServiceMessages.h
+++ b/Engine/Source/Runtime/SessionMessages/Public/SessionServiceMessages.h
@@ -61,10 +61,6 @@ struct FSessionServicePong
 	UPROPERTY(EditAnywhere, Category="Message")
 	FString InstanceName;
 
-	/** Indicates whether the application is running on a console. */
-	UPROPERTY(EditAnywhere, Category="Message")
-	bool IsConsoleBuild;
-
 	/** Holds the name of the platform that the application is running on. */
 	UPROPERTY(EditAnywhere, Category="Message")
 	FString PlatformName;
diff --git a/Engine/Source/Runtime/SessionServices/Private/SessionInstanceInfo.cpp b/Engine/Source/Runtime/SessionServices/Private/SessionInstanceInfo.cpp
--- a/Engine/Source/Runtime/SessionServices/Private/SessionInstanceInfo.cpp
+++ b/Engine/Source/Runtime/SessionServices/Private/SessionInstanceInfo.cpp
@@ -60,7 +60,6 @@ void FSessionInstanceInfo::UpdateFromMessage(const FSessionServicePong& Message,
 	BuildDate = Message.BuildDate;
 	DeviceName = Message.DeviceName;
 	InstanceName = Message.InstanceName;
-	IsConsoleBuild = Message.IsConsoleBuild;
 	PlatformName = Message.PlatformName;
 
 	LastUpdateTime = FDateTime::UtcNow();
diff --git a/Engine/Source/Runtime/SessionServices/Private/SessionInstanceInfo.h b/Engine/Source/Runtime/SessionServices/Private/SessionInstanceInfo.h
--- a/Engine/Source/Runtime/SessionServices/Private/SessionInstanceInfo.h
+++ b/Engine/Source/Runtime/SessionServices/Private/SessionInstanceInfo.h
@@ -125,11 +125,6 @@ public:
 		return Authorized;
 	}
 
-	virtual const bool IsConsole() const override
-	{
-		return IsConsoleBuild;
-	}
-
 	DECLARE_DERIVED_EVENT(FSessionInstanceInfo, ISessionInstanceInfo::FLogReceivedEvent, FLogReceivedEvent);
 	virtual FLogReceivedEvent& OnLogReceived() override
 	{
@@ -183,9 +178,6 @@ private:
 	/** Holds the instance type (i.e. game, editor etc.) */
 	FString InstanceType;
 
-	/** Holds a flag indicating whether this is a console build. */
-	bool IsConsoleBuild;
-
 	/** Holds the time at which the last pong was received. */
 	FDateTime LastUpdateTime;
 
diff --git a/Engine/Source/Runtime/SessionServices/Private/SessionService.cpp b/Engine/Source/Runtime/SessionServices/Private/SessionService.cpp
--- a/Engine/Source/Runtime/SessionServices/Private/SessionService.cpp
+++ b/Engine/Source/Runtime/SessionServices/Private/SessionService.cpp
@@ -123,13 +123,6 @@ void FSessionService::SendPong(const TSharedRef<IMessageContext, ESPMode::Thread
 		Message->DeviceName = FPlatformProcess::ComputerName();
 		Message->InstanceId = FApp::GetInstanceId();
 		Message->InstanceName = FApp::GetInstanceName();
-
-#if PLATFORM_DESKTOP
-		Message->IsConsoleBuild = false;
-#else
-		Message->IsConsoleBuild = true;
-#endif
-
 		Message->PlatformName = FPlatformProperties::PlatformName();
 		Message->SessionId = FApp::GetSessionId();
 		Message->SessionName = FApp::GetSessionName();
diff --git a/Engine/Source/Runtime/SessionServices/Public/ISessionInstanceInfo.h b/Engine/Source/Runtime/SessionServices/Public/ISessionInstanceInfo.h
--- a/Engine/Source/Runtime/SessionServices/Public/ISessionInstanceInfo.h
+++ b/Engine/Source/Runtime/SessionServices/Public/ISessionInstanceInfo.h
@@ -112,13 +112,6 @@ public:
 	 */
 	virtual bool IsAuthorized() const = 0;
 
-	/**
-	 * Check whether this instance is a console build (i.e. no editor features).
-	 *
-	 * @return true if it is a console build, false otherwise.
-	 */
-	virtual const bool IsConsole() const = 0;
-
 	/**
 	 * Checks whether this instance has already begun game play.
 	 *
