Updated session frontend to properly hold on to log messages even if they are currently filtered out

diff --git a/Engine/Source/Developer/SessionFrontend/Private/Widgets/Console/SSessionConsole.cpp b/Engine/Source/Developer/SessionFrontend/Private/Widgets/Console/SSessionConsole.cpp
--- a/Engine/Source/Developer/SessionFrontend/Private/Widgets/Console/SSessionConsole.cpp
+++ b/Engine/Source/Developer/SessionFrontend/Private/Widgets/Console/SSessionConsole.cpp
@@ -500,12 +500,18 @@ void SSessionConsole::HandleSessionManagerInstanceSelectionChanged(const TShared
 
 void SSessionConsole::HandleSessionManagerLogReceived(const TSharedRef<ISessionInfo>& Session, const TSharedRef<ISessionInstanceInfo>& Instance, const TSharedRef<FSessionLogMessage>& Message)
 {
-	if (!SessionManager->IsInstanceSelected(Instance) || !FilterBar->FilterLogMessage(Message))
+	if (!SessionManager->IsInstanceSelected(Instance))
 	{
 		return;
 	}
 
+	// Adds the log to the available messages even if it is filtered out, so that it cannot be lost
 	AvailableLogs.Add(Message);
+	if (!FilterBar->FilterLogMessage(Message))
+	{
+		return;
+	}
+
 	LogMessages.Add(Message);
 
 	LogListView->RequestListRefresh();
