Fix for FrontendSession crash. After much investigation still can't figure out the root cause of this crash, however this trivial change seems to be able to avoid the crash, so let's go with this for now?

diff --git a/Engine/Source/Developer/SessionFrontend/Private/Widgets/Browser/SSessionBrowser.cpp b/Engine/Source/Developer/SessionFrontend/Private/Widgets/Browser/SSessionBrowser.cpp
--- a/Engine/Source/Developer/SessionFrontend/Private/Widgets/Browser/SSessionBrowser.cpp
+++ b/Engine/Source/Developer/SessionFrontend/Private/Widgets/Browser/SSessionBrowser.cpp
@@ -336,7 +336,7 @@ FText SSessionBrowser::HandleSessionTreeRowGetToolTipText(TSharedPtr<FSessionBro
 			ToolTipTextBuilder.AppendLineFormat(LOCTEXT("InstanceToolTipPlatform", "Platform: {0}"), FText::FromString(InstanceInfo->GetPlatformName()));
 			ToolTipTextBuilder.AppendLineFormat(LOCTEXT("InstanceToolTipCurrentLevel", "Current Level: {0}"), FText::FromString(InstanceInfo->GetCurrentLevel()));
 			ToolTipTextBuilder.AppendLineFormat(LOCTEXT("InstanceToolTipWorldTimeSeconds", "World Time: {0}"), FText::AsTimespan(FTimespan::FromSeconds(InstanceInfo->GetWorldTimeSeconds())));
-			ToolTipTextBuilder.AppendLineFormat(LOCTEXT("InstanceToolTipPlayBegun", "Play Has Begun: {0}"), InstanceInfo->PlayHasBegun() ? CoreTexts.Yes : CoreTexts.No);
+			ToolTipTextBuilder.AppendLineFormat(LOCTEXT("InstanceToolTipPlayBegun", "Play Has Begun: {0}"), InstanceInfo->PlayHasBegun() ? FCoreTexts::Get().Yes : FCoreTexts::Get().No);
 			ToolTipTextBuilder.AppendLine();
 			ToolTipTextBuilder.AppendLineFormat(LOCTEXT("SessionToolTipLastUpdateTime", "Last Update Time: {0}"), FText::AsDateTime(InstanceInfo->GetLastUpdateTime()));
 		}
