UnrealFrontend: Disabled old Profiler (UnrealFrontend/SessionFrontend/Profiler).
Old Profiler is deprecated since UE 5.0. Use Trace/UnrealInsights instead.
Added UE_DEPRECATED_PROFILER_ENABLED global define toggle (off by default; this allows to temporarily re-enable Profiler if needed; see UnrealFrontend.Target.cs). Profiler module/code will be removed in a future UE version.

diff --git a/Engine/Source/Developer/SessionFrontend/Private/Widgets/SSessionFrontend.cpp b/Engine/Source/Developer/SessionFrontend/Private/Widgets/SSessionFrontend.cpp
--- a/Engine/Source/Developer/SessionFrontend/Private/Widgets/SSessionFrontend.cpp
+++ b/Engine/Source/Developer/SessionFrontend/Private/Widgets/SSessionFrontend.cpp
@@ -22,7 +22,7 @@
 #include "WorkspaceMenuStructure.h"
 #include "WorkspaceMenuStructureModule.h"
 
-#if STATS
+#if STATS && UE_DEPRECATED_PROFILER_ENABLED
 #include "IProfilerModule.h"
 #endif
 
@@ -73,10 +73,12 @@ void SSessionFrontend::Construct( const FArguments& InArgs, const TSharedRef<SDo
 		.SetIcon(FSlateIcon(FAppStyle::GetAppStyleSetName(), "SessionFrontEnd.Tabs.ScreenComparison"))
 		.SetGroup(AppMenuGroup);
 
+#if STATS && UE_DEPRECATED_PROFILER_ENABLED
 	TabManager->RegisterTabSpawner(ProfilerTabId, FOnSpawnTab::CreateRaw(this, &SSessionFrontend::HandleTabManagerSpawnTab, ProfilerTabId))
 		.SetDisplayName(LOCTEXT("ProfilerTabTitle", "Profiler"))
 		.SetIcon(FSlateIcon(FAppStyle::GetAppStyleSetName(), "Profiler.Tab"))
 		.SetGroup(AppMenuGroup);
+#endif
 	
 	// create tab layout
 	const TSharedRef<FTabManager::FLayout> Layout = FTabManager::NewLayout("SessionFrontendLayout_v1.2")
@@ -99,7 +101,9 @@ void SSessionFrontend::Construct( const FArguments& InArgs, const TSharedRef<SDo
 						->AddTab(SessionConsoleTabId, ETabState::OpenedTab)
 						->AddTab(AutomationTabId, ETabState::OpenedTab)
 						->AddTab(SessionScreenTabId, ETabState::OpenedTab)
+#if STATS && UE_DEPRECATED_PROFILER_ENABLED
 						->AddTab(ProfilerTabId, ETabState::OpenedTab)
+#endif
 						->SetSizeCoefficient(0.75f)
 						->SetForegroundTab(SessionConsoleTabId)
 				)							
@@ -196,7 +200,7 @@ TSharedRef<SDockTab> SSessionFrontend::HandleTabManagerSpawnTab( const FSpawnTab
 
 		AutomationWindowModule.OnShutdown().BindSP(const_cast<SSessionFrontend*>(this), &SSessionFrontend::HandleAutomationModuleShutdown);
 	}
-#if STATS
+#if STATS && UE_DEPRECATED_PROFILER_ENABLED
 	else if (TabIdentifier == ProfilerTabId)
 	{
 		IProfilerModule& ProfilerModule = FModuleManager::LoadModuleChecked<IProfilerModule>(TEXT("Profiler"));
diff --git a/Engine/Source/Developer/SessionFrontend/SessionFrontend.Build.cs b/Engine/Source/Developer/SessionFrontend/SessionFrontend.Build.cs
--- a/Engine/Source/Developer/SessionFrontend/SessionFrontend.Build.cs
+++ b/Engine/Source/Developer/SessionFrontend/SessionFrontend.Build.cs
@@ -1,5 +1,7 @@
 // Copyright Epic Games, Inc. All Rights Reserved.
 
+using System;
+using System.Linq;
 using UnrealBuildTool;
 
 public class SessionFrontend : ModuleRules
@@ -10,17 +12,16 @@ public class SessionFrontend : ModuleRules
 			new string[] {
 				"Core",
 				"Slate",
-                
 			}
 		);
 
-        PrivateDependencyModuleNames.AddRange(
+		PrivateDependencyModuleNames.AddRange(
 			new string[] {
 				"DesktopPlatform",
 				"ApplicationCore",
-                "InputCore",
+				"InputCore",
 				"Json",
-                "SessionServices",
+				"SessionServices",
 				"SlateCore",
 
 				// @todo gmp: remove these dependencies by making the session front-end extensible
@@ -28,13 +29,16 @@ public class SessionFrontend : ModuleRules
 				"ScreenShotComparison",
 				"ScreenShotComparisonTools",
 				"TargetPlatform",
-                "WorkspaceMenuStructure",
+				"WorkspaceMenuStructure",
 			}
 		);
 
-		if (Target.Configuration != UnrealTargetConfiguration.Shipping)
+		if (Target.GlobalDefinitions.Contains("UE_DEPRECATED_PROFILER_ENABLED=1"))
 		{
-			PrivateDependencyModuleNames.Add("Profiler");
+			if (Target.Configuration != UnrealTargetConfiguration.Shipping)
+			{
+				PrivateDependencyModuleNames.Add("Profiler");
+			}
 		}
 
 		PrivateIncludePathModuleNames.AddRange(
diff --git a/Engine/Source/Programs/UnrealFrontend/Private/Commands/UserInterfaceCommand.cpp b/Engine/Source/Programs/UnrealFrontend/Private/Commands/UserInterfaceCommand.cpp
--- a/Engine/Source/Programs/UnrealFrontend/Private/Commands/UserInterfaceCommand.cpp
+++ b/Engine/Source/Programs/UnrealFrontend/Private/Commands/UserInterfaceCommand.cpp
@@ -58,7 +58,9 @@ void FUserInterfaceCommand::Run(  )
 
 	// load optional modules
 	FModuleManager::Get().LoadModule("DeviceManager");
+#if STATS && UE_DEPRECATED_PROFILER_ENABLED
 	FModuleManager::Get().LoadModule("ProfilerClient");
+#endif
 	FModuleManager::Get().LoadModule("ProjectLauncher");
 	FModuleManager::Get().LoadModule("SessionFrontend");
 	FModuleManager::Get().LoadModule("SettingsEditor");
diff --git a/Engine/Source/Programs/UnrealFrontend/UnrealFrontend.Build.cs b/Engine/Source/Programs/UnrealFrontend/UnrealFrontend.Build.cs
--- a/Engine/Source/Programs/UnrealFrontend/UnrealFrontend.Build.cs
+++ b/Engine/Source/Programs/UnrealFrontend/UnrealFrontend.Build.cs
@@ -1,10 +1,12 @@
 // Copyright Epic Games, Inc. All Rights Reserved.
 
+using System;
+using System.Linq;
 using UnrealBuildTool;
 
 public class UnrealFrontend : ModuleRules
 {
-	public UnrealFrontend( ReadOnlyTargetRules Target ) : base(Target)
+	public UnrealFrontend(ReadOnlyTargetRules Target) : base(Target)
 	{
 		PublicIncludePathModuleNames.Add("Launch");
 
@@ -16,21 +18,19 @@ public class UnrealFrontend : ModuleRules
 
 		PrivateDependencyModuleNames.AddRange(
 			new string[] {
+				"ApplicationCore",
 				"AutomationController",
 				"Core",
-				"ApplicationCore",
 				"CoreUObject",
 				"DeviceManager",
 				"LauncherServices",
 				"Messaging",
-                "OutputLog",
-				"Profiler",
-				"ProfilerClient",
-                "ProjectLauncher",
+				"OutputLog",
+				"ProjectLauncher",
 				"Projects",
 				"SessionFrontend",
-				"SessionServices",                
-                "Slate",
+				"SessionServices",
+				"Slate",
 				"SlateCore",
 				"SourceCodeAccess",
 				"StandaloneRenderer",
@@ -71,11 +71,11 @@ public class UnrealFrontend : ModuleRules
 		// @todo: allow for better plug-in support in standalone Slate apps
 		PrivateDependencyModuleNames.AddRange(
 			new string[] {
+				"DirectoryWatcher",
 				"Networking",
 				"Sockets",
+				"TcpMessaging",
 				"UdpMessaging",
-                "TcpMessaging",
-				"DirectoryWatcher"
 			}
 		);
 
@@ -84,5 +84,15 @@ public class UnrealFrontend : ModuleRules
 				"Messaging",
 			}
 		);
+
+		if (Target.GlobalDefinitions.Contains("UE_DEPRECATED_PROFILER_ENABLED=1"))
+		{
+			PrivateDependencyModuleNames.AddRange(
+				new string[] {
+					"Profiler",
+					"ProfilerClient",
+				}
+			);
+		}
 	}
 }
diff --git a/Engine/Source/Programs/UnrealFrontend/UnrealFrontend.Target.cs b/Engine/Source/Programs/UnrealFrontend/UnrealFrontend.Target.cs
--- a/Engine/Source/Programs/UnrealFrontend/UnrealFrontend.Target.cs
+++ b/Engine/Source/Programs/UnrealFrontend/UnrealFrontend.Target.cs
@@ -10,7 +10,7 @@ public class UnrealFrontendTarget : TargetRules
 		Type = TargetType.Program;
 		LinkType = TargetLinkType.Modular;
 		AdditionalPlugins.Add("UdpMessaging");
-        LaunchModuleName = "UnrealFrontend";
+		LaunchModuleName = "UnrealFrontend";
 
 		bCompileAgainstEngine = false;
 		bCompileAgainstCoreUObject = true;
@@ -22,5 +22,8 @@ public class UnrealFrontendTarget : TargetRules
 		bBuildDeveloperTools = true;
 
 		bHasExports = false;
+
+		// Old Profiler (SessionFrontend/Profiler) is deprecated since UE 5.0. Use Trace/UnrealInsights instead.
+		//GlobalDefinitions.Add("UE_DEPRECATED_PROFILER_ENABLED=1");
 	}
 }
diff --git a/Engine/Source/Runtime/Core/Public/Misc/Build.h b/Engine/Source/Runtime/Core/Public/Misc/Build.h
--- a/Engine/Source/Runtime/Core/Public/Misc/Build.h
+++ b/Engine/Source/Runtime/Core/Public/Misc/Build.h
@@ -488,3 +488,9 @@
 #ifndef UE_STATS_MEMORY_PROFILER_ENABLED
 	#define UE_STATS_MEMORY_PROFILER_ENABLED 0
 #endif
+
+// Controls whether old Profiler (UnrealFrontend/SessionFrontend/Profiler) is enabled or not.
+// Old Profiler is deprecated since UE 5.0. Use Trace/UnrealInsights instead.
+#ifndef UE_DEPRECATED_PROFILER_ENABLED
+	#define UE_DEPRECATED_PROFILER_ENABLED 0
+#endif
