Make Trace Control the foreground tab in Session Frontend
Add tab icon for Trace Control.
Add a style sheet in SessionFrontend so it's incons are also visible when running from Insights.

diff --git a/Engine/Source/Developer/SessionFrontend/Private/Models/SessionConsoleCommands.h b/Engine/Source/Developer/SessionFrontend/Private/Models/SessionConsoleCommands.h
--- a/Engine/Source/Developer/SessionFrontend/Private/Models/SessionConsoleCommands.h
+++ b/Engine/Source/Developer/SessionFrontend/Private/Models/SessionConsoleCommands.h
@@ -6,6 +6,7 @@
 #include "Framework/Commands/InputChord.h"
 #include "Styling/AppStyle.h"
 #include "Framework/Commands/Commands.h"
+#include "SessionFrontendStyle.h"
 
 #define LOCTEXT_NAMESPACE "SessionConsoleCommands"
 
@@ -22,7 +23,7 @@ public:
 		: TCommands<FSessionConsoleCommands>(
 			"SessionConsole",
 			NSLOCTEXT("Contexts", "SessionConsole", "Session Console"),
-			NAME_None, FAppStyle::GetAppStyleSetName()
+			NAME_None, FSessionFrontendStyle::GetStyleSetName()
 		)
 	{ }
 
diff --git a/Engine/Source/Developer/SessionFrontend/Private/SessionFrontendModule.cpp b/Engine/Source/Developer/SessionFrontend/Private/SessionFrontendModule.cpp
--- a/Engine/Source/Developer/SessionFrontend/Private/SessionFrontendModule.cpp
+++ b/Engine/Source/Developer/SessionFrontend/Private/SessionFrontendModule.cpp
@@ -11,6 +11,7 @@
 #include "Widgets/Console/SSessionConsole.h"
 #include "Widgets/SSessionFrontend.h"
 #include "ISessionFrontendModule.h"
+#include "SessionFrontendStyle.h"
 #include "Widgets/Docking/SDockTab.h"
 #include "WorkspaceMenuStructure.h"
 #include "WorkspaceMenuStructureModule.h"
@@ -57,16 +58,20 @@ public:
 	
 	virtual void StartupModule() override
 	{
+		FSessionFrontendStyle::Initialize();
+
 		auto& TabSpawnerEntry = FGlobalTabmanager::Get()->RegisterNomadTabSpawner(SessionFrontendTabName, FOnSpawnTab::CreateRaw(this, &FSessionFrontendModule::SpawnSessionFrontendTab))
 			.SetDisplayName(NSLOCTEXT("FSessionFrontendModule", "FrontendTabTitle", "Session Frontend"))
 			.SetTooltipText(NSLOCTEXT("FSessionFrontendModule", "FrontendTooltipText", "Open the Session Frontend tab."))
-			.SetIcon(FSlateIcon(FAppStyle::GetAppStyleSetName(), "SessionFrontEnd.TabIcon"));
+			.SetIcon(FSlateIcon(FSessionFrontendStyle::GetStyleSetName(), "SessionFrontEnd.TabIcon"));
 
 		TabSpawnerEntry.SetGroup(WorkspaceMenu::GetMenuStructure().GetToolsCategory());
 	}
 
 	virtual void ShutdownModule() override
 	{
+		FSessionFrontendStyle::Shutdown();
+
 		FGlobalTabmanager::Get()->UnregisterNomadTabSpawner(SessionFrontendTabName);
 	}
 
diff --git a/Engine/Source/Developer/SessionFrontend/Private/SessionFrontendStyle.cpp b/Engine/Source/Developer/SessionFrontend/Private/SessionFrontendStyle.cpp
new file mode 100644
--- /dev/null
+++ b/Engine/Source/Developer/SessionFrontend/Private/SessionFrontendStyle.cpp
@@ -0,0 +1,98 @@
+// Copyright Epic Games, Inc. All Rights Reserved.
+
+#include "SessionFrontendStyle.h"
+
+#include "Brushes/SlateBoxBrush.h"
+#include "Brushes/SlateImageBrush.h"
+#include "Brushes/SlateNoResource.h"
+#include "Misc/Paths.h"
+#include "Styling/CoreStyle.h"
+#include "Styling/StyleColors.h"
+#include "Styling/SlateStyleRegistry.h"
+#include "Styling/SlateTypes.h"
+#include "Styling/SlateStyleMacros.h"
+#include "Styling/ToolBarStyle.h"
+
+// This is to fix the issue that SlateStyleMacros like IMAGE_BRUSH look for RootToContentDir but StyleSet->RootToContentDir is how this style is set up
+#define RootToContentDir StyleSet->RootToContentDir
+#define RootToCoreContentDir StyleSet->RootToCoreContentDir
+
+TSharedPtr< FSlateStyleSet > FSessionFrontendStyle::StyleSet = nullptr;
+
+// Const icon sizes
+static const FVector2D Icon8x8(8.0f, 8.0f);
+static const FVector2D Icon9x19(9.0f, 19.0f);
+static const FVector2D Icon14x14(14.0f, 14.0f);
+static const FVector2D Icon16x16(16.0f, 16.0f);
+static const FVector2D Icon20x20(20.0f, 20.0f);
+static const FVector2D Icon22x22(22.0f, 22.0f);
+static const FVector2D Icon24x24(24.0f, 24.0f);
+static const FVector2D Icon28x28(28.0f, 28.0f);
+static const FVector2D Icon27x31(27.0f, 31.0f);
+static const FVector2D Icon26x26(26.0f, 26.0f);
+static const FVector2D Icon32x32(32.0f, 32.0f);
+static const FVector2D Icon40x40(40.0f, 40.0f);
+static const FVector2D Icon48x48(48.0f, 48.0f);
+static const FVector2D Icon75x82(75.0f, 82.0f);
+static const FVector2D Icon360x32(360.0f, 32.0f);
+static const FVector2D Icon171x39(171.0f, 39.0f);
+static const FVector2D Icon170x50(170.0f, 50.0f);
+static const FVector2D Icon267x140(170.0f, 50.0f);
+
+void FSessionFrontendStyle::Initialize()
+{
+	LLM_SCOPE_BYNAME(TEXT("SessionFrontend"));
+
+	// Only register once
+	if( StyleSet.IsValid() )
+	{
+		return;
+	}
+
+	StyleSet = MakeShareable(new FSlateStyleSet("SessionFrontendStyle") );
+	
+	StyleSet->SetContentRoot(FPaths::EngineContentDir() / TEXT("Editor/Slate"));
+	StyleSet->SetCoreContentRoot(FPaths::EngineContentDir() / TEXT("Slate"));
+
+	// Session Console tab
+	{
+		StyleSet->Set("SessionConsole.SessionCopy", new IMAGE_BRUSH("Icons/icon_file_open_40x", Icon40x40));
+		StyleSet->Set("SessionConsole.SessionCopy.Small", new IMAGE_BRUSH("Icons/icon_file_open_16px", Icon20x20));
+		StyleSet->Set("SessionConsole.Clear", new IMAGE_BRUSH("Icons/icon_file_new_40x", Icon40x40));
+		StyleSet->Set("SessionConsole.Clear.Small", new IMAGE_BRUSH("Icons/icon_file_new_16px", Icon20x20));
+		StyleSet->Set("SessionConsole.SessionSave", new IMAGE_BRUSH("Icons/icon_file_savelevels_40x", Icon40x40));
+		StyleSet->Set("SessionConsole.SessionSave.Small", new IMAGE_BRUSH("Icons/icon_file_savelevels_16px", Icon20x20));
+	}
+
+	// Session Frontend Window
+	{
+		StyleSet->Set("SessionFrontEnd.TabIcon", new CORE_IMAGE_BRUSH_SVG("Starship/Common/SessionFrontend", Icon16x16));
+		StyleSet->Set("SessionFrontEnd.Tabs.Tools", new CORE_IMAGE_BRUSH("/Icons/icon_tab_Tools_16x", Icon16x16));
+		StyleSet->Set("SessionFrontEnd.Tabs.Console", new CORE_IMAGE_BRUSH_SVG("Starship/Common/Console", Icon16x16));
+		StyleSet->Set("SessionFrontEnd.Tabs.Automation", new CORE_IMAGE_BRUSH_SVG("Starship/Common/AutomationTools", Icon16x16));
+		StyleSet->Set("SessionFrontEnd.Tabs.ScreenComparison", new CORE_IMAGE_BRUSH_SVG("Starship/Common/Monitor", Icon16x16));
+		StyleSet->Set("SessionFrontEnd.Tabs.TraceControl", new IMAGE_BRUSH_SVG("Starship/Common/TraceDataFiltering", Icon16x16));
+	}
+	
+	FSlateStyleRegistry::RegisterSlateStyle( *StyleSet.Get() );
+}
+
+void FSessionFrontendStyle::Shutdown()
+{
+	if( StyleSet.IsValid() )
+	{
+		FSlateStyleRegistry::UnRegisterSlateStyle( *StyleSet.Get() );
+		ensure( StyleSet.IsUnique() );
+		StyleSet.Reset();
+	}
+}
+
+const ISlateStyle& FSessionFrontendStyle::Get()
+{
+	return *( StyleSet.Get() );
+}
+
+const FName& FSessionFrontendStyle::GetStyleSetName()
+{
+	return StyleSet->GetStyleSetName();
+}
\ No newline at end of file
diff --git a/Engine/Source/Developer/SessionFrontend/Private/SessionFrontendStyle.h b/Engine/Source/Developer/SessionFrontend/Private/SessionFrontendStyle.h
new file mode 100644
--- /dev/null
+++ b/Engine/Source/Developer/SessionFrontend/Private/SessionFrontendStyle.h
@@ -0,0 +1,23 @@
+// Copyright Epic Games, Inc. All Rights Reserved.
+
+#pragma once
+
+#include "Styling/SlateStyle.h"
+
+class FSessionFrontendStyle
+{
+public:
+	static void Initialize();
+	static void Shutdown();
+
+	static const ISlateStyle& Get();
+	static const FName& GetStyleSetName();
+	static const FSlateBrush* GetBrush(FName PropertyName, const ANSICHAR* Specifier = NULL)
+	{
+		return StyleSet->GetBrush(PropertyName, Specifier);
+	}
+
+private:
+	/** Singleton instances of this style. */
+	static TSharedPtr<class FSlateStyleSet> StyleSet;	
+};
\ No newline at end of file
diff --git a/Engine/Source/Developer/SessionFrontend/Private/Widgets/SSessionFrontend.cpp b/Engine/Source/Developer/SessionFrontend/Private/Widgets/SSessionFrontend.cpp
--- a/Engine/Source/Developer/SessionFrontend/Private/Widgets/SSessionFrontend.cpp
+++ b/Engine/Source/Developer/SessionFrontend/Private/Widgets/SSessionFrontend.cpp
@@ -17,6 +17,7 @@
 #include "Interfaces/IScreenShotComparisonModule.h"
 #include "Interfaces/ITraceToolsModule.h"
 #include "ISessionServicesModule.h"
+#include "SessionFrontendStyle.h"
 #include "Widgets/Browser/SSessionBrowser.h"
 #include "Widgets/Console/SSessionConsole.h"
 #include "Widgets/Docking/SDockTab.h"
@@ -58,38 +59,38 @@ void SSessionFrontend::Construct( const FArguments& InArgs, const TSharedRef<SDo
 
 	TabManager->RegisterTabSpawner(AutomationTabId, FOnSpawnTab::CreateRaw(this, &SSessionFrontend::HandleTabManagerSpawnTab, AutomationTabId))
 		.SetDisplayName(LOCTEXT("AutomationTabTitle", "Automation"))
-		.SetIcon(FSlateIcon(FAppStyle::GetAppStyleSetName(), "SessionFrontEnd.Tabs.Automation"))
+		.SetIcon(FSlateIcon(FSessionFrontendStyle::GetStyleSetName(), "SessionFrontEnd.Tabs.Automation"))
 		.SetGroup(AppMenuGroup);
 
 	TabManager->RegisterTabSpawner(SessionBrowserTabId, FOnSpawnTab::CreateRaw(this, &SSessionFrontend::HandleTabManagerSpawnTab, SessionBrowserTabId))
 		.SetDisplayName(LOCTEXT("SessionBrowserTitle", "Session Browser"))
-		.SetIcon(FSlateIcon(FAppStyle::GetAppStyleSetName(), "SessionFrontEnd.Tabs.Tools"))
+		.SetIcon(FSlateIcon(FSessionFrontendStyle::GetStyleSetName(), "SessionFrontEnd.Tabs.Tools"))
 		.SetGroup(AppMenuGroup);
 
 	TabManager->RegisterTabSpawner(SessionConsoleTabId, FOnSpawnTab::CreateRaw(this, &SSessionFrontend::HandleTabManagerSpawnTab, SessionConsoleTabId))
 		.SetDisplayName(LOCTEXT("ConsoleTabTitle", "Console"))
-		.SetIcon(FSlateIcon(FAppStyle::GetAppStyleSetName(), "SessionFrontEnd.Tabs.Console"))
+		.SetIcon(FSlateIcon(FSessionFrontendStyle::GetStyleSetName(), "SessionFrontEnd.Tabs.Console"))
 		.SetGroup(AppMenuGroup);
 
 	TabManager->RegisterTabSpawner(TraceControlTabId, FOnSpawnTab::CreateRaw(this, &SSessionFrontend::HandleTabManagerSpawnTab, TraceControlTabId))
 		.SetDisplayName(LOCTEXT("TraceControlTabTitle", "Trace Control"))
-		.SetIcon(FSlateIcon(FAppStyle::GetAppStyleSetName(), "SessionFrontEnd.Tabs.TraceControl"))
+		.SetIcon(FSlateIcon(FSessionFrontendStyle::GetStyleSetName(), "SessionFrontEnd.Tabs.TraceControl"))
 		.SetGroup(AppMenuGroup);
 
 	TabManager->RegisterTabSpawner(SessionScreenTabId, FOnSpawnTab::CreateRaw(this, &SSessionFrontend::HandleTabManagerSpawnTab, SessionScreenTabId))
 		.SetDisplayName(LOCTEXT("ScreenTabTitle", "Screen Comparison"))
-		.SetIcon(FSlateIcon(FAppStyle::GetAppStyleSetName(), "SessionFrontEnd.Tabs.ScreenComparison"))
+		.SetIcon(FSlateIcon(FSessionFrontendStyle::GetStyleSetName(), "SessionFrontEnd.Tabs.ScreenComparison"))
 		.SetGroup(AppMenuGroup);
 
 #if STATS && UE_DEPRECATED_PROFILER_ENABLED
 	TabManager->RegisterTabSpawner(ProfilerTabId, FOnSpawnTab::CreateRaw(this, &SSessionFrontend::HandleTabManagerSpawnTab, ProfilerTabId))
 		.SetDisplayName(LOCTEXT("ProfilerTabTitle", "Profiler"))
-		.SetIcon(FSlateIcon(FAppStyle::GetAppStyleSetName(), "Profiler.Tab"))
+		.SetIcon(FSlateIcon(FSessionFrontendStyle::GetStyleSetName(), "Profiler.Tab"))
 		.SetGroup(AppMenuGroup);
 #endif
 	
 	// create tab layout
-	const TSharedRef<FTabManager::FLayout> Layout = FTabManager::NewLayout("SessionFrontendLayout_v1.3")
+	const TSharedRef<FTabManager::FLayout> Layout = FTabManager::NewLayout("SessionFrontendLayout_v1.4")
 		->AddArea
 		(
 			FTabManager::NewPrimaryArea()
@@ -106,15 +107,15 @@ void SSessionFrontend::Construct( const FArguments& InArgs, const TSharedRef<SDo
 				(
 					// applications
 					FTabManager::NewStack()
-						->AddTab(SessionConsoleTabId, ETabState::OpenedTab)
 						->AddTab(TraceControlTabId, ETabState::OpenedTab)
+						->AddTab(SessionConsoleTabId, ETabState::OpenedTab)
 						->AddTab(AutomationTabId, ETabState::OpenedTab)
 						->AddTab(SessionScreenTabId, ETabState::OpenedTab)
 #if STATS && UE_DEPRECATED_PROFILER_ENABLED
 						->AddTab(ProfilerTabId, ETabState::OpenedTab)
 #endif
 						->SetSizeCoefficient(0.75f)
-						->SetForegroundTab(SessionConsoleTabId)
+						->SetForegroundTab(TraceControlTabId)
 				)							
 		);
 
diff --git a/Engine/Source/Editor/EditorStyle/Private/StarshipStyle.cpp b/Engine/Source/Editor/EditorStyle/Private/StarshipStyle.cpp
--- a/Engine/Source/Editor/EditorStyle/Private/StarshipStyle.cpp
+++ b/Engine/Source/Editor/EditorStyle/Private/StarshipStyle.cpp
@@ -967,23 +967,9 @@ void FStarshipEditorStyle::FStyle::SetupGeneralStyles()
 			.SetShadowColorAndOpacity(FLinearColor(0.0f, 0.0f, 0.0f, 0.9f)));
 	}
 
-	// Session Console tab
-	{
-		Set( "SessionConsole.SessionCopy", new IMAGE_BRUSH( "Icons/icon_file_open_40x", Icon40x40 ) );
-		Set( "SessionConsole.SessionCopy.Small", new IMAGE_BRUSH( "Icons/icon_file_open_16px", Icon20x20 ) );
-		Set( "SessionConsole.Clear", new IMAGE_BRUSH( "Icons/icon_file_new_40x", Icon40x40 ) );
-		Set( "SessionConsole.Clear.Small", new IMAGE_BRUSH( "Icons/icon_file_new_16px", Icon20x20 ) );
-		Set( "SessionConsole.SessionSave", new IMAGE_BRUSH( "Icons/icon_file_savelevels_40x", Icon40x40 ) );
-		Set( "SessionConsole.SessionSave.Small", new IMAGE_BRUSH( "Icons/icon_file_savelevels_16px", Icon20x20 ) );
-	}
-
 	// Session Frontend Window
 	{
-		Set("SessionFrontEnd.TabIcon", new CORE_IMAGE_BRUSH_SVG( "Starship/Common/SessionFrontend", Icon16x16 ) );
-		Set("SessionFrontEnd.Tabs.Tools", new CORE_IMAGE_BRUSH( "/Icons/icon_tab_Tools_16x", Icon16x16 ) );
-		Set("SessionFrontEnd.Tabs.Console", new CORE_IMAGE_BRUSH_SVG("Starship/Common/Console", Icon16x16));
-		Set("SessionFrontEnd.Tabs.Automation", new CORE_IMAGE_BRUSH_SVG("Starship/Common/AutomationTools", Icon16x16));
-		Set("SessionFrontEnd.Tabs.ScreenComparison", new CORE_IMAGE_BRUSH_SVG("Starship/Common/Monitor", Icon16x16));
+		Set("SessionFrontEnd.Tabs.Tools", new CORE_IMAGE_BRUSH("/Icons/icon_tab_Tools_16x", Icon16x16));
 	}
 
 	// Undo History Window
