Improved SessionBrowser selection and expansion.
Added text about the need to run with (-messaging) for sessions to be shown in SessionBrowser.
Changed SessionBrowser row coloring and alignment.
Disable auto selection of current session in Unreal Insights Frontend.
Widened the channel list columns.

diff --git a/Engine/Source/Developer/SessionFrontend/Private/SessionFrontendStyle.cpp b/Engine/Source/Developer/SessionFrontend/Private/SessionFrontendStyle.cpp
--- a/Engine/Source/Developer/SessionFrontend/Private/SessionFrontendStyle.cpp
+++ b/Engine/Source/Developer/SessionFrontend/Private/SessionFrontendStyle.cpp
@@ -73,6 +73,11 @@ void FSessionFrontendStyle::Initialize()
 		StyleSet->Set("SessionFrontEnd.Tabs.ScreenComparison", new CORE_IMAGE_BRUSH_SVG("Starship/Common/Monitor", Icon16x16));
 		StyleSet->Set("SessionFrontEnd.Tabs.TraceControl", new IMAGE_BRUSH_SVG("Starship/Common/TraceDataFiltering", Icon16x16));
 	}
+
+	// Sesssion Browser
+	{
+		StyleSet->Set("SessionBrowser.Row.Name", new IMAGE_BRUSH("ContentBrowser/FilterChecked", FVector2D(7.0f, 24.0f)));
+	}
 	
 	FSlateStyleRegistry::RegisterSlateStyle( *StyleSet.Get() );
 }
diff --git a/Engine/Source/Developer/SessionFrontend/Private/Widgets/Browser/SSessionBrowser.cpp b/Engine/Source/Developer/SessionFrontend/Private/Widgets/Browser/SSessionBrowser.cpp
--- a/Engine/Source/Developer/SessionFrontend/Private/Widgets/Browser/SSessionBrowser.cpp
+++ b/Engine/Source/Developer/SessionFrontend/Private/Widgets/Browser/SSessionBrowser.cpp
@@ -10,7 +10,6 @@
 #include "Widgets/Browser/SSessionBrowserTreeInstanceRow.h"
 #include "Widgets/Browser/SSessionBrowserTreeSessionRow.h"
 
-
 #define LOCTEXT_NAMESPACE "SSessionBrowser"
 
 
@@ -50,40 +49,57 @@ void SSessionBrowser::Construct( const FArguments& InArgs, TSharedRef<ISessionMa
 		+ SVerticalBox::Slot()
 		.FillHeight(1.0f)
 		[
-			// session tree
-			SNew(SBorder)
-			.BorderImage(FAppStyle::GetBrush("ToolPanel.GroupBorder"))
-			.Padding(0.0f)
+			SNew(SOverlay)
+
+			+ SOverlay::Slot()
+			[
+				// session tree
+				SNew(SBorder)
+				.BorderImage(FAppStyle::GetBrush("ToolPanel.GroupBorder"))
+				.Padding(0.0f)
+				[
+					SAssignNew(SessionTreeView, STreeView<TSharedPtr<FSessionBrowserTreeItem>>)
+					.OnExpansionChanged(this, &SSessionBrowser::HandleSessionTreeViewExpansionChanged)
+					.OnGenerateRow(this, &SSessionBrowser::HandleSessionTreeViewGenerateRow)
+					.OnGetChildren(this, &SSessionBrowser::HandleSessionTreeViewGetChildren)
+					.OnSelectionChanged(this, &SSessionBrowser::HandleSessionTreeViewSelectionChanged)
+					.SelectionMode(ESelectionMode::Single)
+					.TreeItemsSource(&SessionTreeItems)
+					.HeaderRow
+					(
+						SNew(SHeaderRow)
+
+						+ SHeaderRow::Column("Name")
+						.DefaultLabel(LOCTEXT("InstanceListNameColumnHeader", "Name"))
+						.FillWidth(0.4f)
+
+						+ SHeaderRow::Column("Type")
+						.DefaultLabel(LOCTEXT("InstanceListTypeColumnHeader", "Type"))
+						.FillWidth(0.15f)
+
+						+ SHeaderRow::Column("Device")
+						.DefaultLabel(LOCTEXT("InstanceListDeviceColumnHeader", "Device"))
+						.FillWidth(0.3f)
+
+						+ SHeaderRow::Column("Status")
+						.DefaultLabel(LOCTEXT("InstanceListStatusColumnHeader", "Status"))
+						.FillWidth(0.15f)
+						.HAlignCell(HAlign_Right)
+						.HAlignHeader(HAlign_Right)
+					)
+				]
+			]
+
+			+ SOverlay::Slot()
+			.VAlign(EVerticalAlignment::VAlign_Bottom)
+			.HAlign(EHorizontalAlignment::HAlign_Left)
+			.Padding(FMargin(3.0f, 0.0f, 3.0f, 5.0f))
 			[
-				SAssignNew(SessionTreeView, STreeView<TSharedPtr<FSessionBrowserTreeItem>>)
-				.OnExpansionChanged(this, &SSessionBrowser::HandleSessionTreeViewExpansionChanged)
-				.OnGenerateRow(this, &SSessionBrowser::HandleSessionTreeViewGenerateRow)
-				.OnGetChildren(this, &SSessionBrowser::HandleSessionTreeViewGetChildren)
-				.OnSelectionChanged(this, &SSessionBrowser::HandleSessionTreeViewSelectionChanged)
-				.SelectionMode(ESelectionMode::Multi)
-				.TreeItemsSource(&SessionTreeItems)
-				.HeaderRow
-				(
-					SNew(SHeaderRow)
-
-					+ SHeaderRow::Column("Name")
-					.DefaultLabel(LOCTEXT("InstanceListNameColumnHeader", "Name"))
-					.FillWidth(0.3f)
-
-					+ SHeaderRow::Column("Type")
-					.DefaultLabel(LOCTEXT("InstanceListTypeColumnHeader", "Type"))
-					.FillWidth(0.2f)
-
-					+ SHeaderRow::Column("Device")
-					.DefaultLabel(LOCTEXT("InstanceListDeviceColumnHeader", "Device"))
-					.FillWidth(0.3f)
-
-					+ SHeaderRow::Column("Status")
-					.DefaultLabel(LOCTEXT("InstanceListStatusColumnHeader", "Status"))
-					.FillWidth(0.2f)
-					.HAlignCell(HAlign_Right)
-					.HAlignHeader(HAlign_Right)
-				)
+				SNew(STextBlock)
+				.AutoWrapText(true)
+				.ColorAndOpacity(FSlateColor(EStyleColor::AccentGray))
+				.Visibility_Lambda([this]() { return ItemMap.Num() < 10 ? EVisibility::Visible : EVisibility::Collapsed; })
+				.Text(LOCTEXT("MessagingInfo", "Sessions must be launched with the \"-messaging\" argument to be visible."))
 			]
 		]
 	];
@@ -116,7 +132,7 @@ END_SLATE_FUNCTION_BUILD_OPTIMIZATION
 
 void SSessionBrowser::ExpandItem(const TSharedPtr<FSessionBrowserTreeItem>& Item)
 {
-	SessionTreeView->SetSingleExpandedItem(Item);
+	SessionTreeView->SetItemExpansion(Item, true);
 
 	if (Item.IsValid())
 	{
@@ -208,7 +224,7 @@ void SSessionBrowser::FilterSessions()
 	// refresh tree view
 	SessionTreeView->RequestTreeRefresh();
 
-	if ( bCanSetDefaultSelection && SessionTreeView->GetNumItemsSelected() == 0 && ThisAppInstance.IsValid() )
+	if ( bCanSetDefaultSelection && SessionTreeView->GetNumItemsSelected() == 0 && ThisAppInstance.IsValid() && !FApp::GetSessionName().Equals("UnrealInsights"))
 	{
 		bCanSetDefaultSelection = false;
 		SessionTreeView->SetItemSelection(ThisAppInstance.Pin(), true, ESelectInfo::Direct);
@@ -429,11 +445,7 @@ void SSessionBrowser::HandleSessionTreeViewSelectionChanged(const TSharedPtr<FSe
 
 				if (InstanceInfo.IsValid())
 				{
-					// special handling for local application
-					if (Item->GetParent() == AppGroupItem)
-					{
-						SessionManager->SelectSession(InstanceInfo->GetOwnerSession());
-					}
+					SessionManager->SelectSession(InstanceInfo->GetOwnerSession());
 
 					SessionManager->SetInstanceSelected(InstanceInfo.ToSharedRef(), true);
 				}
diff --git a/Engine/Source/Developer/SessionFrontend/Private/Widgets/Browser/SSessionBrowserTreeGroupRow.h b/Engine/Source/Developer/SessionFrontend/Private/Widgets/Browser/SSessionBrowserTreeGroupRow.h
--- a/Engine/Source/Developer/SessionFrontend/Private/Widgets/Browser/SSessionBrowserTreeGroupRow.h
+++ b/Engine/Source/Developer/SessionFrontend/Private/Widgets/Browser/SSessionBrowserTreeGroupRow.h
@@ -55,7 +55,7 @@ public:
 				SNew(SBorder)
 					.BorderBackgroundColor(FLinearColor(0.6f, 0.6f, 0.6f, 1.0f))
 					.BorderImage(this, &SSessionBrowserTreeGroupRow::HandleBorderBackgroundImage)
-					.Padding(3.0f)
+					.Padding(1.5f,3.0f)
 					.ToolTipText(Item->GetToolTipText())
 					[
 						SNew(SHorizontalBox)
diff --git a/Engine/Source/Developer/SessionFrontend/Private/Widgets/Browser/SSessionBrowserTreeInstanceRow.h b/Engine/Source/Developer/SessionFrontend/Private/Widgets/Browser/SSessionBrowserTreeInstanceRow.h
--- a/Engine/Source/Developer/SessionFrontend/Private/Widgets/Browser/SSessionBrowserTreeInstanceRow.h
+++ b/Engine/Source/Developer/SessionFrontend/Private/Widgets/Browser/SSessionBrowserTreeInstanceRow.h
@@ -10,6 +10,7 @@
 #include "Widgets/SWidget.h"
 #include "Layout/Margin.h"
 #include "Widgets/Views/STableViewBase.h"
+#include "SessionFrontendStyle.h"
 #include "Styling/StyleDefaults.h"
 #include "Widgets/Layout/SBorder.h"
 #include "Widgets/SBoxPanel.h"
@@ -118,22 +119,28 @@ public:
 
 			if (InstanceInfo.IsValid())
 			{
-				return SNew(SBox)
-					.Padding(FMargin(1.0f, 1.0f, 4.0f, 1.0f))
-					.HAlign(HAlign_Left)
-					[
-						SNew(SBorder)
-						.BorderBackgroundColor(this, &SSessionBrowserTreeInstanceRow::HandleInstanceBorderBackgroundColor)
-						.BorderImage(this, &SSessionBrowserTreeInstanceRow::HandleInstanceBorderBrush)
-						.ColorAndOpacity(FLinearColor(0.25f, 0.25f, 0.25f))
-						.Padding(FMargin(6.0f, 4.0f))
-						.VAlign(VAlign_Center)
-						[
-							SNew(STextBlock)
-							.Font(FAppStyle::GetFontStyle("BoldFont"))
-							.Text(FText::FromString(InstanceInfo->GetInstanceName()))
-						]
-					];
+				return SNew(SHorizontalBox)
+
+				+ SHorizontalBox::Slot()
+				.AutoWidth()
+				.Padding(FMargin(0.0f, 0.0f, 0.0f, 0.0f))
+				.VAlign(VAlign_Center)
+				[
+					SNew(SImage)
+					.Image(FSessionFrontendStyle::GetBrush("SessionBrowser.Row.Name"))
+					.ColorAndOpacity(this, &SSessionBrowserTreeInstanceRow::GetNameImageColor)
+				]
+
+				+ SHorizontalBox::Slot()
+				.AutoWidth()
+				.Padding(FMargin(4.0f, 0.0f, 0.0f, 0.0f))
+				.VAlign(VAlign_Center)
+				[
+					SNew(STextBlock)
+					.Font(FAppStyle::GetFontStyle("BoldFont"))
+					.Text(FText::FromString(InstanceInfo->GetInstanceName()))
+					.ColorAndOpacity(this, &SSessionBrowserTreeInstanceRow::HandleTextColorAndOpacity)
+				];
 			}
 		}
 		else if (ColumnName == "Status")
@@ -151,7 +158,7 @@ public:
 
 			+ SHorizontalBox::Slot()
 			.AutoWidth()
-			.Padding(FMargin(4.0f, 0.0f, 0.0f, 0.0f))
+			.Padding(FMargin(4.0f, 0.0f, 1.5f, 0.0f))
 			.VAlign(VAlign_Center)
 			[
 				SNew(SImage)
@@ -208,8 +215,8 @@ private:
 		return FText::GetEmpty();
 	}
 
-	/** Callback for getting the border color for this row. */
-	FSlateColor HandleInstanceBorderBackgroundColor() const
+	/** Callback for getting the image color for the name column. */
+	FSlateColor GetNameImageColor() const
 	{
 		TSharedPtr<ISessionInstanceInfo> InstanceInfo = Item->GetInstanceInfo();
 
@@ -221,6 +228,19 @@ private:
 		return FLinearColor::Transparent;
 	}
 
+	/** Callback for getting the text color for the name column. */
+	FSlateColor GetNameTextColor() const
+	{
+		FSlateColor Color = GetNameImageColor();
+
+		if (Color == FLinearColor::Transparent)
+		{
+			Color = FSlateColor::UseSubduedForeground();;
+		}
+
+		return Color;
+	}
+
 	/** Callback for getting the border brush for this row. */
 	const FSlateBrush* HandleInstanceBorderBrush() const
 	{
diff --git a/Engine/Source/Developer/TraceTools/Private/Widgets/STraceDataFilterWidget.cpp b/Engine/Source/Developer/TraceTools/Private/Widgets/STraceDataFilterWidget.cpp
--- a/Engine/Source/Developer/TraceTools/Private/Widgets/STraceDataFilterWidget.cpp
+++ b/Engine/Source/Developer/TraceTools/Private/Widgets/STraceDataFilterWidget.cpp
@@ -274,6 +274,7 @@ void STraceDataFilterWidget::ConstructTileView()
 	.OnContextMenuOpening(this, &STraceDataFilterWidget::OnContextMenuOpening)
 	.Orientation(EOrientation::Orient_Horizontal)
 	.ItemHeight(15.0f)
+	.ItemWidth(150.0f)
 	.ListItemsSource(&FilteredListItems);
 }
 
