[UBT][CppCompileWarnings] Engine compliance for wundef

diff --git a/Engine/Source/Developer/SessionFrontend/Private/Widgets/SSessionFrontend.cpp b/Engine/Source/Developer/SessionFrontend/Private/Widgets/SSessionFrontend.cpp
--- a/Engine/Source/Developer/SessionFrontend/Private/Widgets/SSessionFrontend.cpp
+++ b/Engine/Source/Developer/SessionFrontend/Private/Widgets/SSessionFrontend.cpp
@@ -24,7 +24,7 @@
 #include "WorkspaceMenuStructure.h"
 #include "WorkspaceMenuStructureModule.h"
 
-#if STATS && UE_DEPRECATED_PROFILER_ENABLED
+#if STATS && defined(UE_DEPRECATED_PROFILER_ENABLED) && UE_DEPRECATED_PROFILER_ENABLED
 #include "IProfilerModule.h"
 #endif
 
@@ -82,7 +82,7 @@ void SSessionFrontend::Construct( const FArguments& InArgs, const TSharedRef<SDo
 		.SetIcon(FSlateIcon(FSessionFrontendStyle::GetStyleSetName(), "SessionFrontEnd.Tabs.ScreenComparison"))
 		.SetGroup(AppMenuGroup);
 
-#if STATS && UE_DEPRECATED_PROFILER_ENABLED
+#if STATS && defined(UE_DEPRECATED_PROFILER_ENABLED) && UE_DEPRECATED_PROFILER_ENABLED
 	TabManager->RegisterTabSpawner(ProfilerTabId, FOnSpawnTab::CreateRaw(this, &SSessionFrontend::HandleTabManagerSpawnTab, ProfilerTabId))
 		.SetDisplayName(LOCTEXT("ProfilerTabTitle", "Profiler"))
 		.SetIcon(FSlateIcon(FSessionFrontendStyle::GetStyleSetName(), "Profiler.Tab"))
@@ -111,7 +111,7 @@ void SSessionFrontend::Construct( const FArguments& InArgs, const TSharedRef<SDo
 						->AddTab(SessionConsoleTabId, ETabState::OpenedTab)
 						->AddTab(AutomationTabId, ETabState::OpenedTab)
 						->AddTab(SessionScreenTabId, ETabState::OpenedTab)
-#if STATS && UE_DEPRECATED_PROFILER_ENABLED
+#if STATS && defined(UE_DEPRECATED_PROFILER_ENABLED) && UE_DEPRECATED_PROFILER_ENABLED
 						->AddTab(ProfilerTabId, ETabState::OpenedTab)
 #endif
 						->SetSizeCoefficient(0.75f)
@@ -211,7 +211,7 @@ TSharedRef<SDockTab> SSessionFrontend::HandleTabManagerSpawnTab( const FSpawnTab
 
 		AutomationWindowModule.OnShutdown().BindSP(const_cast<SSessionFrontend*>(this), &SSessionFrontend::HandleAutomationModuleShutdown);
 	}
-#if STATS && UE_DEPRECATED_PROFILER_ENABLED
+#if STATS && defined(UE_DEPRECATED_PROFILER_ENABLED) && UE_DEPRECATED_PROFILER_ENABLED
 	else if (TabIdentifier == ProfilerTabId)
 	{
 		IProfilerModule& ProfilerModule = FModuleManager::LoadModuleChecked<IProfilerModule>(TEXT("Profiler"));
diff --git a/Engine/Source/Runtime/Launch/Private/LaunchEngineLoop.cpp b/Engine/Source/Runtime/Launch/Private/LaunchEngineLoop.cpp
--- a/Engine/Source/Runtime/Launch/Private/LaunchEngineLoop.cpp
+++ b/Engine/Source/Runtime/Launch/Private/LaunchEngineLoop.cpp
@@ -227,7 +227,7 @@
 
 #if !(UE_BUILD_SHIPPING || UE_BUILD_TEST)
 	#include "ProfileVisualizerModule.h"
-#if UE_DEPRECATED_PROFILER_ENABLED
+#if defined(UE_DEPRECATED_PROFILER_ENABLED) && UE_DEPRECATED_PROFILER_ENABLED
 	#include "IProfilerServiceModule.h"
 #endif // UE_DEPRECATED_PROFILER_ENABLED
 #endif
@@ -1837,7 +1837,7 @@ int32 FEngineLoop::PreInitPreStartupScreen(const TCHAR* CmdLine)
 	}
 #endif
 
-#if STATS && UE_STATS_MEMORY_PROFILER_ENABLED
+#if STATS && defined(UE_STATS_MEMORY_PROFILER_ENABLED) && UE_STATS_MEMORY_PROFILER_ENABLED
 	// Create the stats malloc profiler proxy.
 	if (FStatsMallocProfilerProxy::HasMemoryProfilerToken())
 	{
@@ -4281,7 +4281,7 @@ int32 FEngineLoop::PreInitPostStartupScreen(const TCHAR* CmdLine)
 	// Initialize profile visualizers.
 #if !(UE_BUILD_SHIPPING || UE_BUILD_TEST)
 	FModuleManager::Get().LoadModule(TEXT("ProfileVisualizer"));
-#if UE_DEPRECATED_PROFILER_ENABLED
+#if defined(UE_DEPRECATED_PROFILER_ENABLED) && UE_DEPRECATED_PROFILER_ENABLED
 	if (FPlatformProcess::SupportsMultithreading())
 	{
 		FModuleManager::Get().LoadModule(TEXT("ProfilerService"));
