[Insights] Frontend: Fixed crash when hovering a session in the Session Frontend (sessions list).
* Fixed usage of FCoreTexts.
* Added default initialization for members of FSessionInstanceInfo.

diff --git a/Engine/Source/Developer/SessionFrontend/Private/Widgets/Browser/SSessionBrowser.cpp b/Engine/Source/Developer/SessionFrontend/Private/Widgets/Browser/SSessionBrowser.cpp
--- a/Engine/Source/Developer/SessionFrontend/Private/Widgets/Browser/SSessionBrowser.cpp
+++ b/Engine/Source/Developer/SessionFrontend/Private/Widgets/Browser/SSessionBrowser.cpp
@@ -343,6 +343,7 @@ FText SSessionBrowser::HandleSessionTreeRowGetToolTipText(TSharedPtr<FSessionBro
 		if (InstanceInfo.IsValid())
 		{
 			const FCoreTexts& CoreTexts = FCoreTexts::Get();
+			const FText PlayHasBegunText = InstanceInfo->PlayHasBegun() ? CoreTexts.Yes : CoreTexts.No;
 
 			ToolTipTextBuilder.AppendLineFormat(LOCTEXT("InstanceToolTipInstanceId", "Instance ID: {0}"), FText::FromString(InstanceInfo->GetInstanceId().ToString(EGuidFormats::DigitsWithHyphensInBraces)));
 			ToolTipTextBuilder.AppendLine();
@@ -351,7 +352,7 @@ FText SSessionBrowser::HandleSessionTreeRowGetToolTipText(TSharedPtr<FSessionBro
 			ToolTipTextBuilder.AppendLineFormat(LOCTEXT("InstanceToolTipPlatform", "Platform: {0}"), FText::FromString(InstanceInfo->GetPlatformName()));
 			ToolTipTextBuilder.AppendLineFormat(LOCTEXT("InstanceToolTipCurrentLevel", "Current Level: {0}"), FText::FromString(InstanceInfo->GetCurrentLevel()));
 			ToolTipTextBuilder.AppendLineFormat(LOCTEXT("InstanceToolTipWorldTimeSeconds", "World Time: {0}"), FText::AsTimespan(FTimespan::FromSeconds(InstanceInfo->GetWorldTimeSeconds())));
-			ToolTipTextBuilder.AppendLineFormat(LOCTEXT("InstanceToolTipPlayBegun", "Play Has Begun: {0}"), InstanceInfo->PlayHasBegun() ? FCoreTexts::Get().Yes : FCoreTexts::Get().No);
+			ToolTipTextBuilder.AppendLineFormat(LOCTEXT("InstanceToolTipPlayBegun", "Play Has Begun: {0}"), PlayHasBegunText);
 			ToolTipTextBuilder.AppendLine();
 			ToolTipTextBuilder.AppendLineFormat(LOCTEXT("SessionToolTipLastUpdateTime", "Last Update Time: {0}"), FText::AsDateTime(InstanceInfo->GetLastUpdateTime()));
 		}
diff --git a/Engine/Source/Runtime/SessionServices/Private/SessionInstanceInfo.h b/Engine/Source/Runtime/SessionServices/Private/SessionInstanceInfo.h
--- a/Engine/Source/Runtime/SessionServices/Private/SessionInstanceInfo.h
+++ b/Engine/Source/Runtime/SessionServices/Private/SessionInstanceInfo.h
@@ -152,7 +152,7 @@ private:
 	FMessageAddress ApplicationAddress;
 
 	/** Whether the current user is authorized to interact with this instance. */
-	bool Authorized;
+	bool Authorized = false;
 
 	/** Holds the instance's build date. */
 	FString BuildDate;
@@ -167,10 +167,10 @@ private:
 	FMessageAddress EngineAddress;
 
 	/** Holds the instance's engine version. */
-	int32 EngineVersion;
+	int32 EngineVersion = 0;
 
 	/** Holds a flag indicating whether the game has begun. */
-	bool HasBegunPlay;
+	bool HasBegunPlay = false;
 
 	/** Holds the instance identifier. */
 	FGuid InstanceId;
@@ -197,7 +197,7 @@ private:
 	FString PlatformName;
 
 	/** Holds the instance's current game world time. */
-	float WorldTimeSeconds;
+	float WorldTimeSeconds = 0.0f;
 
 private:
 
