* Removed deprecated code (toggled off by UE_DEPRECATED_PROFILER_ENABLED or by UE_STATS_MEMORY_PROFILER_ENABLED).
* Removed dependencies to deprecated (and already removed) Profiler* modules.

diff --git a/Engine/Plugins/Runtime/NetworkPredictionInsights/Source/NetworkPredictionInsights/Private/UI/NetworkPredictionInsightsManager.h b/Engine/Plugins/Runtime/NetworkPredictionInsights/Source/NetworkPredictionInsights/Private/UI/NetworkPredictionInsightsManager.h
--- a/Engine/Plugins/Runtime/NetworkPredictionInsights/Source/NetworkPredictionInsights/Private/UI/NetworkPredictionInsightsManager.h
+++ b/Engine/Plugins/Runtime/NetworkPredictionInsights/Source/NetworkPredictionInsights/Private/UI/NetworkPredictionInsightsManager.h
@@ -50,11 +50,8 @@ protected:
 
 public:
 	/**
-	 * @return the global instance of the Networking Profiler (Networking Insights) manager.
-	 * This is an internal singleton and cannot be used outside ProfilerModule.
-	 * For external use:
-	 *     IProfilerModule& ProfilerModule = FModuleManager::Get().LoadModuleChecked<IProfilerModule>("Profiler");
-	 *     ProfilerModule.GetProfilerManager();
+	 * @return the global instance of the Network Prediction Insights manager.
+	 * This is an internal singleton and cannot be used outside NetworkPredictionInsights module.
 	 */
 	static TSharedPtr<FNetworkPredictionInsightsManager> Get();
 
diff --git a/Engine/Source/Developer/SessionFrontend/Private/Widgets/SSessionFrontend.cpp b/Engine/Source/Developer/SessionFrontend/Private/Widgets/SSessionFrontend.cpp
--- a/Engine/Source/Developer/SessionFrontend/Private/Widgets/SSessionFrontend.cpp
+++ b/Engine/Source/Developer/SessionFrontend/Private/Widgets/SSessionFrontend.cpp
@@ -24,10 +24,6 @@
 #include "WorkspaceMenuStructure.h"
 #include "WorkspaceMenuStructureModule.h"
 
-#if STATS && defined(UE_DEPRECATED_PROFILER_ENABLED) && UE_DEPRECATED_PROFILER_ENABLED
-#include "IProfilerModule.h"
-#endif
-
 
 #define LOCTEXT_NAMESPACE "SSessionFrontend"
 
@@ -41,8 +37,6 @@ static const FName SessionConsoleTabId("SessionConsole");
 static const FName SessionScreenTabId("ScreenComparison");
 static const FName TraceControlTabId("TraceControl");
 
-static const FName ProfilerTabId("Profiler");
-
 
 /* SSessionFrontend interface
  *****************************************************************************/
@@ -82,13 +76,6 @@ void SSessionFrontend::Construct( const FArguments& InArgs, const TSharedRef<SDo
 		.SetIcon(FSlateIcon(FSessionFrontendStyle::GetStyleSetName(), "SessionFrontEnd.Tabs.ScreenComparison"))
 		.SetGroup(AppMenuGroup);
 
-#if STATS && defined(UE_DEPRECATED_PROFILER_ENABLED) && UE_DEPRECATED_PROFILER_ENABLED
-	TabManager->RegisterTabSpawner(ProfilerTabId, FOnSpawnTab::CreateRaw(this, &SSessionFrontend::HandleTabManagerSpawnTab, ProfilerTabId))
-		.SetDisplayName(LOCTEXT("ProfilerTabTitle", "Profiler"))
-		.SetIcon(FSlateIcon(FSessionFrontendStyle::GetStyleSetName(), "Profiler.Tab"))
-		.SetGroup(AppMenuGroup);
-#endif
-	
 	// create tab layout
 	const TSharedRef<FTabManager::FLayout> Layout = FTabManager::NewLayout("SessionFrontendLayout_v1.4")
 		->AddArea
@@ -111,9 +98,6 @@ void SSessionFrontend::Construct( const FArguments& InArgs, const TSharedRef<SDo
 						->AddTab(SessionConsoleTabId, ETabState::OpenedTab)
 						->AddTab(AutomationTabId, ETabState::OpenedTab)
 						->AddTab(SessionScreenTabId, ETabState::OpenedTab)
-#if STATS && defined(UE_DEPRECATED_PROFILER_ENABLED) && UE_DEPRECATED_PROFILER_ENABLED
-						->AddTab(ProfilerTabId, ETabState::OpenedTab)
-#endif
 						->SetSizeCoefficient(0.75f)
 						->SetForegroundTab(TraceControlTabId)
 				)							
@@ -211,13 +195,6 @@ TSharedRef<SDockTab> SSessionFrontend::HandleTabManagerSpawnTab( const FSpawnTab
 
 		AutomationWindowModule.OnShutdown().BindSP(const_cast<SSessionFrontend*>(this), &SSessionFrontend::HandleAutomationModuleShutdown);
 	}
-#if STATS && defined(UE_DEPRECATED_PROFILER_ENABLED) && UE_DEPRECATED_PROFILER_ENABLED
-	else if (TabIdentifier == ProfilerTabId)
-	{
-		IProfilerModule& ProfilerModule = FModuleManager::LoadModuleChecked<IProfilerModule>(TEXT("Profiler"));
-		TabWidget = ProfilerModule.CreateProfilerWindow(SessionManager.ToSharedRef(), DockTab);
-	}
-#endif
 	else if (TabIdentifier == SessionBrowserTabId)
 	{
 		TabWidget = SNew(SSessionBrowser, SessionManager.ToSharedRef());
diff --git a/Engine/Source/Developer/SessionFrontend/SessionFrontend.Build.cs b/Engine/Source/Developer/SessionFrontend/SessionFrontend.Build.cs
--- a/Engine/Source/Developer/SessionFrontend/SessionFrontend.Build.cs
+++ b/Engine/Source/Developer/SessionFrontend/SessionFrontend.Build.cs
@@ -35,14 +35,6 @@ public class SessionFrontend : ModuleRules
 			}
 		);
 
-		if (Target.GlobalDefinitions.Contains("UE_DEPRECATED_PROFILER_ENABLED=1"))
-		{
-			if (Target.Configuration != UnrealTargetConfiguration.Shipping)
-			{
-				PrivateDependencyModuleNames.Add("Profiler");
-			}
-		}
-
 		PrivateIncludePathModuleNames.AddRange(
 			new string[] {
 				"Messaging",
diff --git a/Engine/Source/Editor/UnrealEd/Private/EditorEngine.cpp b/Engine/Source/Editor/UnrealEd/Private/EditorEngine.cpp
--- a/Engine/Source/Editor/UnrealEd/Private/EditorEngine.cpp
+++ b/Engine/Source/Editor/UnrealEd/Private/EditorEngine.cpp
@@ -1483,7 +1483,6 @@ void UEditorEngine::LoadDefaultEditorModules()
 			TEXT("AutomationWindow"),
 			TEXT("AutomationController"),
 			TEXT("DeviceManager"),
-			TEXT("ProfilerClient"),
 			TEXT("SessionFrontend"),
 			TEXT("LegacyProjectLauncher"),
 			TEXT("SettingsEditor"),
diff --git a/Engine/Source/Runtime/Launch/Launch.Build.cs b/Engine/Source/Runtime/Launch/Launch.Build.cs
--- a/Engine/Source/Runtime/Launch/Launch.Build.cs
+++ b/Engine/Source/Runtime/Launch/Launch.Build.cs
@@ -189,12 +189,6 @@ public class Launch : ModuleRules
 			DynamicallyLoadedModuleNames.AddRange(new string[] {
 				"ProfileVisualizer",
 			});
-
-			if (Target.GlobalDefinitions.Contains("UE_DEPRECATED_PROFILER_ENABLED=1"))
-			{
-				PublicIncludePathModuleNames.Add("ProfilerService");
-				DynamicallyLoadedModuleNames.Add("ProfilerService");
-			}
 		}
 
 		// The engine can use AutomationController in any connfiguration besides shipping.  This module is loaded
@@ -206,12 +200,6 @@ public class Launch : ModuleRules
 
 		if (Target.bBuildEditor == true)
 		{
-			if (Target.GlobalDefinitions.Contains("UE_DEPRECATED_PROFILER_ENABLED=1"))
-			{
-				PublicIncludePathModuleNames.Add("ProfilerClient");
-				DynamicallyLoadedModuleNames.Add("ProfilerClient");
-			}
-
 			PrivateDependencyModuleNames.AddRange(new string[] {
 					"SourceControl",
 					"EditorFramework",
diff --git a/Engine/Source/Runtime/Launch/Private/LaunchEngineLoop.cpp b/Engine/Source/Runtime/Launch/Private/LaunchEngineLoop.cpp
--- a/Engine/Source/Runtime/Launch/Private/LaunchEngineLoop.cpp
+++ b/Engine/Source/Runtime/Launch/Private/LaunchEngineLoop.cpp
@@ -227,9 +227,6 @@
 
 #if !(UE_BUILD_SHIPPING || UE_BUILD_TEST)
 	#include "ProfileVisualizerModule.h"
-#if defined(UE_DEPRECATED_PROFILER_ENABLED) && UE_DEPRECATED_PROFILER_ENABLED
-	#include "IProfilerServiceModule.h"
-#endif // UE_DEPRECATED_PROFILER_ENABLED
 #endif
 
 #if WITH_AUTOMATION_WORKER
@@ -1837,19 +1834,6 @@ int32 FEngineLoop::PreInitPreStartupScreen(const TCHAR* CmdLine)
 	}
 #endif
 
-#if STATS && defined(UE_STATS_MEMORY_PROFILER_ENABLED) && UE_STATS_MEMORY_PROFILER_ENABLED
-	// Create the stats malloc profiler proxy.
-	if (FStatsMallocProfilerProxy::HasMemoryProfilerToken())
-	{
-		if (PLATFORM_USES_FIXED_GMalloc_CLASS)
-		{
-			UE_LOG(LogMemory, Fatal, TEXT("Cannot do malloc profiling with PLATFORM_USES_FIXED_GMalloc_CLASS."));
-		}
-		// Assumes no concurrency here.
-		UE::Private::GMalloc = FStatsMallocProfilerProxy::Get();
-	}
-#endif // STATS && UE_STATS_MEMORY_PROFILER_ENABLED
-
 #if WITH_APPLICATION_CORE
 	{
 		SCOPED_BOOT_TIMING("CreateConsoleOutputDevice");
@@ -4281,13 +4265,6 @@ int32 FEngineLoop::PreInitPostStartupScreen(const TCHAR* CmdLine)
 	// Initialize profile visualizers.
 #if !(UE_BUILD_SHIPPING || UE_BUILD_TEST)
 	FModuleManager::Get().LoadModule(TEXT("ProfileVisualizer"));
-#if defined(UE_DEPRECATED_PROFILER_ENABLED) && UE_DEPRECATED_PROFILER_ENABLED
-	if (FPlatformProcess::SupportsMultithreading())
-	{
-		FModuleManager::Get().LoadModule(TEXT("ProfilerService"));
-		FModuleManager::Get().GetModuleChecked<IProfilerServiceModule>("ProfilerService").CreateProfilerServiceManager();
-	}
-#endif // UE_DEPRECATED_PROFILER_ENABLED
 #endif
 
 	// Init HighRes screenshot system, unless running on server
@@ -4866,11 +4843,6 @@ int32 FEngineLoop::Init()
 #endif
 
 #if WITH_EDITOR
-	if (GIsEditor)
-	{
-		FModuleManager::Get().LoadModule(TEXT("ProfilerClient"));
-	}
-
 	FModuleManager::Get().LoadModule(TEXT("SequenceRecorder"));
 	FModuleManager::Get().LoadModule(TEXT("SequenceRecorderSections"));
 #endif
